RemoteControlWebInterface: Looking for nodejs executable in default install folder if it is not in the PATH env variable

diff --git a/Engine/Plugins/VirtualProduction/RemoteControlWebInterface/WebApp/Start.bat b/Engine/Plugins/VirtualProduction/RemoteControlWebInterface/WebApp/Start.bat
--- a/Engine/Plugins/VirtualProduction/RemoteControlWebInterface/WebApp/Start.bat
+++ b/Engine/Plugins/VirtualProduction/RemoteControlWebInterface/WebApp/Start.bat
@@ -1,14 +1,20 @@
 @echo off
 
-rem First we check if nodejs is installed
-for %%X in (node.exe) do (SET node=%%~$PATH:X)
+rem   Check if nodejs is in the env variable PATH 
+for %%X in (node.exe) do (set node=%%~$PATH:X)
+if not defined node (
+  echo Did not find nodejs in PATH, looking in default installation folder
+  if exist "%ProgramFiles%\nodejs\node.exe" (set node="%ProgramFiles%\nodejs\node.exe")
+)
+
 if not defined node (
   echo ERROR: Couldn't find node.js installed, Please install latest nodejs from https://nodejs.org/en/download/
   exit 1
 )
 
+
 rem Let's check if it is a modern nodejs
-node -e "process.exit( process.versions.node.split('.')[0] );"
+%node% -e "process.exit( process.versions.node.split('.')[0] );"
 echo Found Node.js version %errorlevel% (%node%)
 
 if %errorlevel% LSS 8 (
@@ -17,4 +23,4 @@ if %errorlevel% LSS 8 (
 )
 
 rem redirecting all command line arguments to node script
-node Scripts/start.js %*
\ No newline at end of file
+"%node%" Scripts/start.js %*
\ No newline at end of file
