RemoteControlWebInterface: WIP - styles, vector drawer, rebind widgets, drag empty widgets

diff --git a/Engine/Plugins/VirtualProduction/RemoteControlWebInterface/WebApp/Client/src/components/ActorWidgets/index.ts b/Engine/Plugins/VirtualProduction/RemoteControlWebInterface/WebApp/Client/src/components/ActorWidgets/index.ts
deleted file mode 100644
diff --git a/Engine/Plugins/VirtualProduction/RemoteControlWebInterface/WebApp/Client/src/components/Screens/index.ts b/Engine/Plugins/VirtualProduction/RemoteControlWebInterface/WebApp/Client/src/components/Screens/index.ts
new file mode 100644
--- /dev/null
+++ b/Engine/Plugins/VirtualProduction/RemoteControlWebInterface/WebApp/Client/src/components/Screens/index.ts
@@ -0,0 +1 @@
+export * from './Snapshot';
diff --git a/Engine/Plugins/VirtualProduction/RemoteControlWebInterface/WebApp/Client/src/components/Widgets/index.ts b/Engine/Plugins/VirtualProduction/RemoteControlWebInterface/WebApp/Client/src/components/Widgets/index.ts
--- a/Engine/Plugins/VirtualProduction/RemoteControlWebInterface/WebApp/Client/src/components/Widgets/index.ts
+++ b/Engine/Plugins/VirtualProduction/RemoteControlWebInterface/WebApp/Client/src/components/Widgets/index.ts
@@ -1,7 +1,7 @@
-export * from './SlidersWidget';
-export * from './SliderWidget';
 export * from './DropdownWidget';
 export * from './ImageSelectorWidget';
+export * from './SlidersWidget';
+export * from './SliderWidget';
+export * from './TabWidget';
 export * from './ToggleWidget';
-export * from './VectorWidget';
-export * from './ColorSwatchWidget';
\ No newline at end of file
+export * from './VectorWidget';
\ No newline at end of file
diff --git a/Engine/Plugins/VirtualProduction/RemoteControlWebInterface/WebApp/Client/src/components/controls/index.ts b/Engine/Plugins/VirtualProduction/RemoteControlWebInterface/WebApp/Client/src/components/controls/index.ts
--- a/Engine/Plugins/VirtualProduction/RemoteControlWebInterface/WebApp/Client/src/components/controls/index.ts
+++ b/Engine/Plugins/VirtualProduction/RemoteControlWebInterface/WebApp/Client/src/components/controls/index.ts
@@ -1,13 +1,8 @@
 export * from './Button';
-export * from './ColorPicker';
-export * from './ColorPickerExtended';
 export * from './Gauge';
 export * from './GaugesWrapper';
-export * from './Invalid';
 export * from './JoysticksWrapper';
-export * from './Level';
-export * from './MultiSlider';
-export * from './Sequence';
+export * from './NewColorPicker';
 export * from './Slider';
 export * from './SlidersWrapper';
 export * from './Text';
diff --git a/Engine/Plugins/VirtualProduction/RemoteControlWebInterface/WebApp/Client/src/components/index.ts b/Engine/Plugins/VirtualProduction/RemoteControlWebInterface/WebApp/Client/src/components/index.ts
--- a/Engine/Plugins/VirtualProduction/RemoteControlWebInterface/WebApp/Client/src/components/index.ts
+++ b/Engine/Plugins/VirtualProduction/RemoteControlWebInterface/WebApp/Client/src/components/index.ts
@@ -1,19 +1,10 @@
-export * from './ActorWidget';
 export * from './AlertModal';
 export * from './ConnectionStatus';
 export * from './CustomStack';
+export * from './DraggableWidget';
 export * from './IconModal';
-export * from './LayoutConfig';
 export * from './PresetModal';
 export * from './RenameModal';
 export * from './SafeImage';
-export * from './Tabs';
-export * from './Widget';
-export * from './WidgetConfig';
-export * from './WidgetPreview';
-export * from './WidgetsPanel';
-export * from './WidgetsPanelWrapper';
-export * from './WidgetsStack';
-export * from './WidgetSettings';
-export * from './WidgetWrapper';
-export * from './WidgetDnd';
\ No newline at end of file
+export * from './Screen';
+export * from './Tabs';
\ No newline at end of file
diff --git a/Engine/Plugins/VirtualProduction/RemoteControlWebInterface/WebApp/Client/src/reducers/api.ts b/Engine/Plugins/VirtualProduction/RemoteControlWebInterface/WebApp/Client/src/reducers/api.ts
--- a/Engine/Plugins/VirtualProduction/RemoteControlWebInterface/WebApp/Client/src/reducers/api.ts
+++ b/Engine/Plugins/VirtualProduction/RemoteControlWebInterface/WebApp/Client/src/reducers/api.ts
@@ -145,6 +145,14 @@ export const _api = {
     },
     asset: (asset: string, action: AssetAction, meta?: any) => {
       _socket.emit('asset', asset, action, meta);
+    },
+    rename: (type: 'property' | 'function', property: string, label: string, callback: (label: string) => void) => {
+      const preset = _internal.getPreset();
+      _socket.emit('rename', preset, type, property, label, callback);
+    },
+    metadata: (property: string, meta: string, value: string, callback: () => void) => {
+      const preset = _internal.getPreset();
+      _socket.emit('metadata', preset, property, meta, value, callback);
     }
   },
   actor: {
diff --git a/Engine/Plugins/VirtualProduction/RemoteControlWebInterface/WebApp/Client/src/shared/index.ts b/Engine/Plugins/VirtualProduction/RemoteControlWebInterface/WebApp/Client/src/shared/index.ts
--- a/Engine/Plugins/VirtualProduction/RemoteControlWebInterface/WebApp/Client/src/shared/index.ts
+++ b/Engine/Plugins/VirtualProduction/RemoteControlWebInterface/WebApp/Client/src/shared/index.ts
@@ -72,6 +72,7 @@ export interface IExposedFunction {
   DisplayName: string;
   UnderlyingFunction: IFunction;
   Metadata: { [key: string]: string };
+  Widget: WidgetType;
 }
 
 export interface IProperty {
@@ -141,6 +142,7 @@ export enum WidgetTypes {
   Dropdown =        'Dropdown',
   ImageSelector =   'Image Selector',
   Vector =          'Vector',
+  Spacer =          'Spacer',
 
   Level =           'Level',
   Sequence =        'Sequence',
@@ -150,18 +152,10 @@ export type WidgetType = keyof typeof WidgetTypes | string;
 
 
 export type IWidgetMeta = { [key: string]: any } & {
-  default?: PropertyValue;
-  min?: number;
-  max?: number;
+  Min?: number;
+  Max?: number;
 };
 
-export interface IWidget {
-  title?: string;
-  type?: WidgetType;
-  property?: string;
-  meta?: IWidgetMeta;
-  order?: number;
-}
 
 export enum IPanelType {
   Panel = 'PANEL',
@@ -173,25 +167,17 @@ export interface IPanel {
   title?: string;
   type: IPanelType;
   widgets?: ICustomStackWidget[];
-  
   items?: ICustomStackListItem[];
-  addFunction?: ICustomStackFunction;
-  removeFunction?: ICustomStackFunction;
 }
 
 export enum TabLayout {
-  Panel1X1 =      '1x1',
-  Panel1x2 =      '1x2',
-  Panel2x1 =      '2x1',
-  Panel2x2 =      '2x2',
-  Stack =         'Stack',
-
-  CustomActor =   'CustomActor',
+  Stack =    'Stack',
+  Screen =   'Screen',
 }
 
-export interface IActorWidget {
-  type: 'Walls' | 'Cards' | 'Location' | 'Camera' | 'Save' | 'Another' | 'GreenScreen' | 'Snapshot';
-  actors: string[];
+export interface IScreen {
+  type: 'Snapshot';
+  data: any;
 }
 
 export interface IDropdownOption {
@@ -201,15 +187,10 @@ export interface IDropdownOption {
 
 export interface ICustomStackProperty {
   id?: string;
-  actor?: string;
   property: string;
   propertyType: PropertyType;
   widget: WidgetType;
   label?: string;
-  reset?: PropertyValue;
-
-  min?: number;
-  max?: number;
 
   // Sliders & Vector only
   lock?: boolean;
@@ -221,11 +202,12 @@ export interface ICustomStackProperty {
 
   // Dropdown only
   options?: IDropdownOption[];
-}
 
-export interface ICustomStackFunction {
-  actor: string;
-  function: string;
+  // Space only
+  spaces?: number;
+
+  // Dial only
+  mode?: 'ENDLESS' | 'RANGE' | 'ROTATION';
 }
 
 export interface ICustomStackItem {
@@ -249,16 +231,14 @@ export interface ICustomStackListItem {
 
 export type ICustomStackWidget = ICustomStackProperty | ICustomStackTabs;
 
+export interface IView {
+  tabs: ITab[];
+}
+
 export interface ITab {
   name: string;
   icon: string;
   layout: TabLayout;
   panels?: IPanel[];
-  stack?: ICustomStackWidget[];
-
-  actor?: IActorWidget;
-}
-
-export interface IView {
-  tabs: ITab[];
+  screen?: IScreen;
 }
\ No newline at end of file
diff --git a/Engine/Plugins/VirtualProduction/RemoteControlWebInterface/WebApp/Client/src/utilities/WidgetUtilities.ts b/Engine/Plugins/VirtualProduction/RemoteControlWebInterface/WebApp/Client/src/utilities/WidgetUtilities.ts
--- a/Engine/Plugins/VirtualProduction/RemoteControlWebInterface/WebApp/Client/src/utilities/WidgetUtilities.ts
+++ b/Engine/Plugins/VirtualProduction/RemoteControlWebInterface/WebApp/Client/src/utilities/WidgetUtilities.ts
@@ -1,5 +1,5 @@
-import { IExposedFunction, IExposedProperty, IPreset, IWidget, PropertyType, WidgetType, WidgetTypes } from "src/shared";
-import _ from 'lodash';
+import { IPreset, PropertyType, WidgetType, WidgetTypes } from "src/shared";
+
 
 type Range = {
   min?: number;
@@ -7,7 +7,6 @@ type Range = {
 }
 
 export type WidgetProperties<TPropertyValue = any> = {
-  widget?: IWidget;
   stack?: boolean;
   type?: PropertyType;
   value?: TPropertyValue;
@@ -24,29 +23,8 @@ export class WidgetUtilities {
   static propertyWidgets: Record<string, CustomWidget[]> = {};
   static customWidgets: Record<string, CustomWidget> = {};
   
-  static getDefaultRange(property: PropertyType, widget: WidgetType): Range {
-    switch (property) {
-      case PropertyType.Rotator:
-        return { min: 0, max: 360	};
-
-      case PropertyType.LinearColor:
-      case PropertyType.Vector4:
-        return { min: 0, max: 1 };
-
-      case PropertyType.Color:
-        return { min: 0, max: 255 };
-
-      case PropertyType.Vector:
-        if (widget === WidgetTypes.Joystick)
-          return { min: -100000, max: 100000 };
-        break;
-    }
-
-    return { min: 0, max: 100 };
-  }
-
-  static getPropertyLimits(property: IExposedProperty): Range {
-    switch (property.UnderlyingProperty.Type) {
+  static getPropertyLimits(type: PropertyType): Range {
+    switch (type) {
       case PropertyType.Int8:
         return { min: -128, max: 127 };
 
@@ -66,6 +44,7 @@ export class WidgetUtilities {
       case PropertyType.Vector:
       case PropertyType.Vector2D:
       case PropertyType.Vector4:
+      case PropertyType.LinearColor:
         return { min: -1.17549e+38, max: 3.402823466E+38 };
 
       case PropertyType.Rotator:
@@ -74,110 +53,26 @@ export class WidgetUtilities {
       case PropertyType.Double:
         return { min: Number.MIN_VALUE, max: Number.MAX_VALUE };
 
-      case PropertyType.LinearColor:
-        return { min: 0, max: 1 };
-
       case PropertyType.Color:
         return { min: 0, max: 255 };
     }
-  }
-
-  static getMinMax(preset: IPreset, widget: IWidget): Range {
-    const property = WidgetUtilities.getProperty(preset, widget) as IExposedProperty;
-    if (!property?.UnderlyingProperty)
-      return;
-
-    // Property hard limits
-    const limits = WidgetUtilities.getPropertyLimits(property);
-
-    // Get range from UE meta data fields
-    const prop = WidgetUtilities.verifyInRange({
-                    min: parseFloat(property.Metadata?.Min),
-                    max: parseFloat(property.Metadata?.Max),
-                  },
-                  limits);
-
-    let range = WidgetUtilities.verifyInRange(widget.meta, limits);
-
-    // Make sure it is in the allowed prop range
-    range = WidgetUtilities.verifyInRange(range, prop);
-
-    const defRange = WidgetUtilities.getDefaultRange(property.UnderlyingProperty.Type, widget.type);
-
-    return {
-      min: range.min ?? prop.min ?? defRange.min ?? limits.min,
-      max: range.max ?? prop.max ?? defRange.max ?? limits.max,
-    };
-  }
-
-  static verifyInRange(range: Range, limits: Range): Range {
-    const ret: Range = {};
 
-    if (!isNaN(range?.min))
-      ret.min = isNaN(limits?.min) ? range.min : Math.max(range.min, limits.min);
-
-    if (!isNaN(range?.max))
-      ret.max = isNaN(limits?.max) ? range.max : Math.min(range.max, limits.max);
-
-    return ret;
-  }
-
-  static getProperty(preset: IPreset, widget: IWidget): IExposedProperty|IExposedFunction {
-    if (!preset || !widget?.property)
-      return;
-
-    return preset.Exposed[widget.property];
+    return {};
   }
 
-  static getPropertyType(preset: IPreset, widget: IWidget): PropertyType {
-    const property = WidgetUtilities.getProperty(preset, widget) as IExposedProperty;
-    return property?.UnderlyingProperty?.Type;
+  static parseNumber(value: string): number {
+    const number = parseFloat(value);
+    if (!isNaN(number))
+      return number;
   }
 
-  static getDefaultWidgetProps(propertyType: PropertyType) {
-    switch (propertyType) {
-      case PropertyType.LinearColor:
-        return { R: true, G: true, B: true, value: { R: 1, G: 1, B: 1 } };
-
-      case PropertyType.Color:
-        return { R: true, G: true, B: true, value: { R: 255, G: 255, B: 255 } };
-
-      case PropertyType.Vector2D:
-        return { X: true, Y: true, value: { X: 1, Y: 1 } };
-
-      case PropertyType.Vector:
-        return { X: true, Y: true, Z: true, value: { X: 1, Y: 1, Z: 1 } };
-
-      case PropertyType.Vector4:
-        return { X: true, Y: true, Z: true, W: false, value: { X: 1, Y: 1, Z: 1, W: 1 } };
+  static getMinMax(preset: IPreset, exposed: string): Range {
+    const property = preset.Exposed[exposed];
 
-      case PropertyType.Rotator:
-        return { X: true, Y: true, value: { Roll: 360, Pitch: 360 } };
-    }
-
-    return { value: 0 };
-  }
-
-  static getSubProperties(type: PropertyType) {
-    switch (type) {
-      case PropertyType.Vector2D:
-        return ['X', 'Y'];
-
-      case PropertyType.Vector:
-        return ['X', 'Y', 'Z'];
-
-      case PropertyType.Vector4:
-        return ['X', 'Y', 'Z', 'W'];
-
-      case PropertyType.Rotator:
-        return ['X', 'Y', 'Z'];
-
-      case PropertyType.Color:
-      case PropertyType.LinearColor:
-        return ['R', 'G', 'B', 'A'];
-    }
-
-    return [];
+    return {
+      min: WidgetUtilities.parseNumber(property?.Metadata?.Min),
+      max: WidgetUtilities.parseNumber(property?.Metadata?.Max),
+    };
   }
 
   static isAsset(type: WidgetType): boolean {
@@ -190,73 +85,6 @@ export class WidgetUtilities {
     return false;
   }
 
-  static compatibleWidgets(type: PropertyType | string): WidgetType[] {
-    const widgets = [];
-
-    switch (type) {
-      case PropertyType.Boolean:
-      case PropertyType.Uint8:
-        widgets.push(WidgetTypes.Toggle);
-        break;
-
-      case PropertyType.Int8:
-      case PropertyType.Int16:
-      case PropertyType.Int32:
-      case PropertyType.Int64:
-      case PropertyType.Uint16:
-      case PropertyType.Uint32:
-      case PropertyType.Uint64:
-      case PropertyType.Float:
-      case PropertyType.Double:
-        widgets.push(WidgetTypes.Gauge, WidgetTypes.Slider);
-        break;
-
-      case PropertyType.Color:
-      case PropertyType.LinearColor:
-        widgets.push(WidgetTypes.ColorPicker, WidgetTypes.Sliders);
-        break;
-
-      case PropertyType.Vector:
-        widgets.push(WidgetTypes.Joystick, WidgetTypes.Sliders, WidgetTypes.ScaleSlider);
-        break;
-
-      case PropertyType.Vector2D:
-        widgets.push(WidgetTypes.Sliders, WidgetTypes.ScaleSlider);
-        break;
-
-      case PropertyType.Vector4:
-        widgets.push(WidgetTypes.ColorPicker, WidgetTypes.Sliders);
-        break;
-
-      case PropertyType.Rotator:
-        widgets.push(WidgetTypes.Gauge, WidgetTypes.Joystick, WidgetTypes.Sliders);
-        break;
-
-      case PropertyType.Function:
-        widgets.push(WidgetTypes.Button);
-        break;
-
-      case PropertyType.String:
-      case PropertyType.Text:
-        widgets.push(WidgetTypes.Text);
-        break;
-
-      case 'World':
-        widgets.push(WidgetTypes.Level);
-        break;
-
-      case 'LevelSequence':
-        widgets.push(WidgetTypes.Sequence);
-        break;
-    }
-
-    const custom = WidgetUtilities.propertyWidgets[type];
-    if (custom)
-      widgets.push(...custom.map(custom => custom.type));
-
-    return widgets;
-  }
-
   static registerWidget(type: string, properties: PropertyType[], render: (props: WidgetProperties) => React.ReactNode): void {
     const custom: CustomWidget = { type, render };
     WidgetUtilities.customWidgets[type] = custom;
@@ -269,60 +97,40 @@ export class WidgetUtilities {
     }
   }
 
-  static getPropertyTypeData(propertyType: PropertyType) {
-    let keys = [];
-    let min = 0;
-    let max = 100;
-    let percision = 2;
-
+  static getPropertyPercision(propertyType: PropertyType) {
     switch (propertyType) {
-      case PropertyType.Int8:
-      case PropertyType.Int16:
-      case PropertyType.Int32:
-        min = -100;
-        max = 100;
-        percision = 0;
-        break;
-
       case PropertyType.Float:
       case PropertyType.Double:
-        min = -100;
-        max = 100;
-        break;
+      case PropertyType.Vector:
+      case PropertyType.Vector2D:
+      case PropertyType.Vector4:
+      case PropertyType.Rotator:
+      case PropertyType.LinearColor:
+        return 3;
+    }
 
-      case PropertyType.Uint16:
-      case PropertyType.Uint32:
-        min = 0;
-        max = 100;
-        percision = 0;
-        break;
+    return 0;
+  }
 
+  static getPropertyKeys(propertyType: PropertyType) {
+    switch (propertyType) {
       case PropertyType.Rotator:
-        keys = ['Roll', 'Pitch', 'Yaw'];
-        break;
+        return ['Roll', 'Pitch', 'Yaw'];
 
       case PropertyType.Vector:
-        keys = ['X', 'Y', 'Z'];
-        break;
+        return ['X', 'Y', 'Z'];
 
       case PropertyType.Vector2D:
-        keys = ['X', 'Y'];
-        break;
+        return ['X', 'Y'];
 
       case PropertyType.Vector4:
-        keys = ['X', 'Y', 'Z', 'W'];
-        break;
-
-      case PropertyType.LinearColor:
-        keys = ['R', 'G', 'B'];
-        break;
+        return ['X', 'Y', 'Z', 'W'];
 
       case PropertyType.Color:
-        max = 255;
-        keys = ['R', 'G', 'B'];
-        break;
+      case PropertyType.LinearColor:
+        return ['R', 'G', 'B'];
     }
 
-    return { keys, min, max, percision };
+    return [];
   };
 }
\ No newline at end of file
diff --git a/Engine/Plugins/VirtualProduction/RemoteControlWebInterface/WebApp/Server/src/Notify.ts b/Engine/Plugins/VirtualProduction/RemoteControlWebInterface/WebApp/Server/src/Notify.ts
--- a/Engine/Plugins/VirtualProduction/RemoteControlWebInterface/WebApp/Server/src/Notify.ts
+++ b/Engine/Plugins/VirtualProduction/RemoteControlWebInterface/WebApp/Server/src/Notify.ts
@@ -17,7 +17,9 @@ export namespace Notify {
         .on('actor', UnrealEngine.setActorValue)
         .on('reset', UnrealEngine.resetPayloadValue)
         .on('execute', UnrealEngine.executeFunction)
-        .on('asset', UnrealEngine.executeAssetAction);
+        .on('asset', UnrealEngine.executeAssetAction)
+        .on('rename', UnrealEngine.renamePresetProperty)
+        .on('metadata', UnrealEngine.setPresetPropertyMetadata);
 
       if (UnrealEngine.isConnected())
         socket.emit('connected', true);
diff --git a/Engine/Plugins/VirtualProduction/RemoteControlWebInterface/WebApp/Server/src/UnrealEngine.ts b/Engine/Plugins/VirtualProduction/RemoteControlWebInterface/WebApp/Server/src/UnrealEngine.ts
--- a/Engine/Plugins/VirtualProduction/RemoteControlWebInterface/WebApp/Server/src/UnrealEngine.ts
+++ b/Engine/Plugins/VirtualProduction/RemoteControlWebInterface/WebApp/Server/src/UnrealEngine.ts
@@ -60,6 +60,10 @@ namespace UnrealApi {
   export type Assets = {
     Assets: IAsset[];
   };
+
+  export type RenameLabel = {
+    AssignedLabel: string;
+  };
 }
 
 export namespace UnrealEngine {
@@ -141,8 +145,12 @@ export namespace UnrealEngine {
             break;
 
           for (const field of message.ChangedFields) {
-            setPayloadValueInternal(payloads, [message.PresetName, field.PropertyLabel], field.PropertyValue);
-            Notify.emitValueChange(message.PresetName, field.PropertyLabel, field.PropertyValue);
+            const property = _.find(preset.ExposedProperties, p => p.DisplayName === field.PropertyLabel); 
+            if (!property)
+              continue;
+
+            setPayloadValueInternal(payloads, [message.PresetName, property.Id], field.PropertyValue);
+            Notify.emitValueChange(message.PresetName, property.Id, field.PropertyValue);
           }
           break;
         }
@@ -235,70 +243,9 @@ export namespace UnrealEngine {
       const { Presets } = await get<UnrealApi.Presets>('/remote/presets');
       
       for (const p of Presets ?? []) {
-        if (!_.find(presets, preset => preset.Name === p.Name)) {
-          registerToPreset(p.Name);
-
-          const res = await get<UnrealApi.View>(`/remote/preset/${p.Name}/metadata/view`);
-          refreshView(p.Name, res?.Value);
-        }
-
-        const { Preset } = await get<UnrealApi.Preset>(`/remote/preset/${p.Name}`);
-        if (!Preset)
-          continue;
-
-        Preset.ExposedProperties = [];
-        Preset.ExposedFunctions = [];
-        Preset.Exposed = {};
-
-        for (const Group of Preset.Groups) {
-          for (const Function of Group.ExposedFunctions)
-            Preset.Exposed[Function.Id] = Function;
-
-          for (const Property of Group.ExposedProperties) {
-            Preset.Exposed[Property.Id] = Property;
-
-            switch (Property.UnderlyingProperty.Type) {
-              case PropertyType.Boolean:
-              case PropertyType.Uint8:
-              case PropertyType.Int8:
-                Property.Widget = WidgetTypes.Toggle;
-                break;
-
-              case PropertyType.Int16:
-              case PropertyType.Int32:
-              case PropertyType.Int64:
-              case PropertyType.Uint16:
-              case PropertyType.Uint32:
-              case PropertyType.Uint64:
-              case PropertyType.Float:
-              case PropertyType.Double:
-                Property.Widget = WidgetTypes.Slider;
-                break;
-
-              case PropertyType.LinearColor:
-              case PropertyType.Color:
-              case PropertyType.Vector4:
-                Property.Widget = WidgetTypes.ColorPicker;
-                break;
-
-              case PropertyType.Vector:
-              case PropertyType.Vector2D:
-              case PropertyType.Rotator:
-                Property.Widget = WidgetTypes.Vector;
-                break;
-
-              case PropertyType.String:
-              case PropertyType.Text:
-                Property.Widget = WidgetTypes.Text;
-                break;
-            }
-          }
-
-          Preset.ExposedProperties.push(...Group.ExposedProperties);
-          Preset.ExposedFunctions.push(...Group.ExposedFunctions);
-        }
-        
-        all.push(Preset);
+        const preset = await pullPreset(p.Name);
+        if (preset)
+          all.push(preset);
       }
 
       const compact =  _.compact(all);
@@ -311,6 +258,82 @@ export namespace UnrealEngine {
     }
   }
 
+  async function pullPreset(name: string): Promise<IPreset> {
+    try {
+      if (!_.find(presets, preset => preset.Name === name)) {
+        registerToPreset(name);
+
+        const res = await get<UnrealApi.View>(`/remote/preset/${name}/metadata/view`);
+        refreshView(name, res?.Value);
+      }
+
+      const { Preset } = await get<UnrealApi.Preset>(`/remote/preset/${name}`);
+      if (!Preset)
+        return null;
+
+      Preset.ExposedProperties = [];
+      Preset.ExposedFunctions = [];
+      Preset.Exposed = {};
+
+      for (const Group of Preset.Groups) {
+        for (const Function of Group.ExposedFunctions) {
+          Function.Widget = WidgetTypes.Button;
+          Preset.Exposed[Function.Id] = Function;
+        }
+
+        for (const Property of Group.ExposedProperties) {
+          Preset.Exposed[Property.Id] = Property;
+
+          switch (Property.UnderlyingProperty.Type) {
+            case PropertyType.Boolean:
+            case PropertyType.Uint8:
+            case PropertyType.Int8:
+              Property.Widget = WidgetTypes.Toggle;
+              break;
+
+            case PropertyType.Int16:
+            case PropertyType.Int32:
+            case PropertyType.Int64:
+            case PropertyType.Uint16:
+            case PropertyType.Uint32:
+            case PropertyType.Uint64:
+            case PropertyType.Float:
+            case PropertyType.Double:
+              Property.Widget = WidgetTypes.Slider;
+              break;
+
+            case PropertyType.LinearColor:
+            case PropertyType.Color:
+            case PropertyType.Vector4:
+              Property.Widget = WidgetTypes.ColorPicker;
+              break;
+
+            case PropertyType.Vector:
+            case PropertyType.Vector2D:
+            case PropertyType.Rotator:
+              Property.Widget = WidgetTypes.Vector;
+              break;
+
+            case PropertyType.String:
+            case PropertyType.Text:
+              Property.Widget = WidgetTypes.Text;
+              break;
+          }
+
+          if (Property.UnderlyingProperty.Type.startsWith('TEnum'))
+            Property.Widget = WidgetTypes.Dropdown;
+        }
+
+        Preset.ExposedProperties.push(...Group.ExposedProperties);
+        Preset.ExposedFunctions.push(...Group.ExposedFunctions);
+      }
+
+      return Preset;
+    } catch (error) {
+      console.log(`Failed to pull preset '${name}' data`);
+    }
+  }
+
   async function refreshView(preset: string, viewJson: string) {
     try {
       if (!viewJson)
@@ -530,6 +553,28 @@ export namespace UnrealEngine {
     }
   }
 
+  export async function renamePresetProperty(preset: string, type: string, property: string, label: string, onChanged: (value: string) => void): Promise<void> {
+    try {
+      const url = `/remote/preset/${preset}/${type}/${property}/label`;
+      const res = await put<UnrealApi.RenameLabel>(url, { NewLabel: label });
+      await refresh();
+      onChanged?.(res?.AssignedLabel ?? label);
+
+    } catch (err) {
+      console.log(`Failed to rename a property`);
+    }
+  }
+
+  export async function setPresetPropertyMetadata(preset: string, property: string, metadata: string, value: string) {
+    try {
+
+      const url = `/remote/preset/${preset}/property/${property}/metadata/${metadata}`;
+      await put(url, { value });
+    } catch (err) {
+      console.log(`Failed to rename a property`);
+    }
+  }
+
   export async function executeAssetAction(asset: string, action: AssetAction, meta: any): Promise<void> {
     if (!asset)
       return;
diff --git a/Engine/Plugins/VirtualProduction/RemoteControlWebInterface/WebApp/package.json b/Engine/Plugins/VirtualProduction/RemoteControlWebInterface/WebApp/package.json
--- a/Engine/Plugins/VirtualProduction/RemoteControlWebInterface/WebApp/package.json
+++ b/Engine/Plugins/VirtualProduction/RemoteControlWebInterface/WebApp/package.json
@@ -1,6 +1,6 @@
 {
   "name": "conductor",
-  "version": "0.3.1",
+  "version": "0.3.7",
   "scripts": {
     "postinstall": "concurrently \"cd ./Client && npm install && cd ..\" \"cd ./Server && npm install && cd ..\"",
     "build": "node ./Scripts/build.js",
