RemoteControlWebInterface: WIP, Vector Drawer & Joysticks

diff --git a/Engine/Plugins/VirtualProduction/RemoteControlWebInterface/Source/RemoteControlWebInterface/Private/RCWebInterfaceCustomizations.cpp b/Engine/Plugins/VirtualProduction/RemoteControlWebInterface/Source/RemoteControlWebInterface/Private/RCWebInterfaceCustomizations.cpp
--- a/Engine/Plugins/VirtualProduction/RemoteControlWebInterface/Source/RemoteControlWebInterface/Private/RCWebInterfaceCustomizations.cpp
+++ b/Engine/Plugins/VirtualProduction/RemoteControlWebInterface/Source/RemoteControlWebInterface/Private/RCWebInterfaceCustomizations.cpp
@@ -42,7 +42,7 @@ namespace RCWebInterface
 		{
 			if (Property->IsA<FNumericProperty>() && !Property->IsA<FByteProperty>())
 			{
-				Widgets = { TEXT("Slider") };
+				Widgets = { TEXT("Slider"), TEXT("Dial") };
 			}
 			else if (Property->IsA<FBoolProperty>())
 			{
@@ -61,7 +61,7 @@ namespace RCWebInterface
 				}
 				else
 				{
-					Widgets = { TEXT("Slider") };
+					Widgets = { TEXT("Slider"), TEXT("Dial") };
 				}
 			}
 			else if (Property->IsA<FEnumProperty>())
@@ -73,9 +73,21 @@ namespace RCWebInterface
 				FStructProperty* StructProperty = CastField<FStructProperty>(Property);
 				if (UStruct* Struct = StructProperty->Struct)
 				{
-					if (Struct->IsChildOf(TBaseStructure<FVector>::Get()) || Struct->IsChildOf(TBaseStructure<FRotator>::Get()))
+					if (Struct->IsChildOf(TBaseStructure<FVector>::Get()))
 					{
-						Widgets = { TEXT("Vector"), TEXT("Sliders") };
+						Widgets = { TEXT("Vector"), TEXT("Joystick"), TEXT("Dials"), TEXT("Sliders") };
+					}
+					else if (Struct->IsChildOf(TBaseStructure<FVector2D>::Get()))
+					{
+						Widgets = { TEXT("Vector"), TEXT("Joystick"), TEXT("Dials"), TEXT("Sliders") };
+					}
+					else if (Struct->IsChildOf(TBaseStructure<FRotator>::Get()))
+					{
+						Widgets = { TEXT("Vector"), TEXT("Dials"), TEXT("Sliders") };
+					}
+					else if (Struct->IsChildOf(TBaseStructure<FVector4>::Get()) || Struct->IsChildOf(TBaseStructure<FColor>::Get()) || Struct->IsChildOf(TBaseStructure<FLinearColor>::Get()))
+					{
+						Widgets = { TEXT("Color Picker") };
 					}
 				}
 			}
diff --git a/Engine/Plugins/VirtualProduction/RemoteControlWebInterface/WebApp/Client/src/components/controls/index.ts b/Engine/Plugins/VirtualProduction/RemoteControlWebInterface/WebApp/Client/src/components/controls/index.ts
--- a/Engine/Plugins/VirtualProduction/RemoteControlWebInterface/WebApp/Client/src/components/controls/index.ts
+++ b/Engine/Plugins/VirtualProduction/RemoteControlWebInterface/WebApp/Client/src/components/controls/index.ts
@@ -1,6 +1,7 @@
 export * from './Button';
 export * from './Gauge';
 export * from './GaugesWrapper';
+export * from './Joystick';
 export * from './JoysticksWrapper';
 export * from './NewColorPicker';
 export * from './Slider';
diff --git a/Engine/Plugins/VirtualProduction/RemoteControlWebInterface/WebApp/Client/src/shared/index.ts b/Engine/Plugins/VirtualProduction/RemoteControlWebInterface/WebApp/Client/src/shared/index.ts
--- a/Engine/Plugins/VirtualProduction/RemoteControlWebInterface/WebApp/Client/src/shared/index.ts
+++ b/Engine/Plugins/VirtualProduction/RemoteControlWebInterface/WebApp/Client/src/shared/index.ts
@@ -72,7 +72,6 @@ export interface IExposedFunction {
   DisplayName: string;
   UnderlyingFunction: IFunction;
   Metadata: { [key: string]: string };
-  Widget: WidgetType;
 }
 
 export interface IProperty {
@@ -196,7 +195,7 @@ export interface ICustomStackProperty {
   lock?: boolean;
 
   // Vector only
-  widgets?: (WidgetTypes.Joystick | WidgetTypes.Sliders | WidgetTypes.Gauge)[];
+  widgets?: WidgetType[];
   speedMin?: number;
   speedMax?: number;
 
diff --git a/Engine/Plugins/VirtualProduction/RemoteControlWebInterface/WebApp/Server/src/UnrealEngine.ts b/Engine/Plugins/VirtualProduction/RemoteControlWebInterface/WebApp/Server/src/UnrealEngine.ts
--- a/Engine/Plugins/VirtualProduction/RemoteControlWebInterface/WebApp/Server/src/UnrealEngine.ts
+++ b/Engine/Plugins/VirtualProduction/RemoteControlWebInterface/WebApp/Server/src/UnrealEngine.ts
@@ -14,6 +14,7 @@ namespace UnrealApi {
     FieldsRemoved     = 'PresetFieldsRemoved',
     MetadataModified  = 'PresetMetadataModified',
     ActorModified     = 'PresetActorModified',
+    EntitiesModified  = 'PresetEntitiesModified',
   }
 
   export type Presets = { 
@@ -160,6 +161,9 @@ export namespace UnrealEngine {
           refresh();
           break;
 
+        case UnrealApi.PresetEvent.EntitiesModified:
+          break;
+
         case UnrealApi.PresetEvent.MetadataModified:
           if (!message.Metadata.view)
             break;
@@ -276,54 +280,11 @@ export namespace UnrealEngine {
       Preset.Exposed = {};
 
       for (const Group of Preset.Groups) {
-        for (const Function of Group.ExposedFunctions) {
-          Function.Widget = WidgetTypes.Button;
+        for (const Function of Group.ExposedFunctions)
           Preset.Exposed[Function.Id] = Function;
-        }
-
-        for (const Property of Group.ExposedProperties) {
+        for (const Property of Group.ExposedProperties)
           Preset.Exposed[Property.Id] = Property;
 
-          switch (Property.UnderlyingProperty.Type) {
-            case PropertyType.Boolean:
-            case PropertyType.Uint8:
-            case PropertyType.Int8:
-              Property.Widget = WidgetTypes.Toggle;
-              break;
-
-            case PropertyType.Int16:
-            case PropertyType.Int32:
-            case PropertyType.Int64:
-            case PropertyType.Uint16:
-            case PropertyType.Uint32:
-            case PropertyType.Uint64:
-            case PropertyType.Float:
-            case PropertyType.Double:
-              Property.Widget = WidgetTypes.Slider;
-              break;
-
-            case PropertyType.LinearColor:
-            case PropertyType.Color:
-            case PropertyType.Vector4:
-              Property.Widget = WidgetTypes.ColorPicker;
-              break;
-
-            case PropertyType.Vector:
-            case PropertyType.Vector2D:
-            case PropertyType.Rotator:
-              Property.Widget = WidgetTypes.Vector;
-              break;
-
-            case PropertyType.String:
-            case PropertyType.Text:
-              Property.Widget = WidgetTypes.Text;
-              break;
-          }
-
-          if (Property.UnderlyingProperty.Type.startsWith('TEnum'))
-            Property.Widget = WidgetTypes.Dropdown;
-        }
-
         Preset.ExposedProperties.push(...Group.ExposedProperties);
         Preset.ExposedFunctions.push(...Group.ExposedFunctions);
       }
diff --git a/Engine/Plugins/VirtualProduction/RemoteControlWebInterface/WebApp/package.json b/Engine/Plugins/VirtualProduction/RemoteControlWebInterface/WebApp/package.json
--- a/Engine/Plugins/VirtualProduction/RemoteControlWebInterface/WebApp/package.json
+++ b/Engine/Plugins/VirtualProduction/RemoteControlWebInterface/WebApp/package.json
@@ -1,6 +1,6 @@
 {
   "name": "conductor",
-  "version": "0.3.7",
+  "version": "0.3.8",
   "scripts": {
     "postinstall": "concurrently \"cd ./Client && npm install && cd ..\" \"cd ./Server && npm install && cd ..\"",
     "build": "node ./Scripts/build.js",
