RemoteControlWebInterface: WIP, Vector Drawer & Joysticks

diff --git a/Engine/Plugins/VirtualProduction/RemoteControlWebInterface/WebApp/Client/src/components/LayoutTabDrawer/index.ts b/Engine/Plugins/VirtualProduction/RemoteControlWebInterface/WebApp/Client/src/components/LayoutTabDrawer/index.ts
--- a/Engine/Plugins/VirtualProduction/RemoteControlWebInterface/WebApp/Client/src/components/LayoutTabDrawer/index.ts
+++ b/Engine/Plugins/VirtualProduction/RemoteControlWebInterface/WebApp/Client/src/components/LayoutTabDrawer/index.ts
@@ -1,5 +1,4 @@
 export * from './LayoutTabDrawer';
-export * from './LayoutTabSearch';
 export * from '../Widgets/SlidersWidget';
 export * from '../Widgets/SliderWidget';
 export * from '../Widgets/DropdownWidget';
diff --git a/Engine/Plugins/VirtualProduction/RemoteControlWebInterface/WebApp/Client/src/components/controls/index.ts b/Engine/Plugins/VirtualProduction/RemoteControlWebInterface/WebApp/Client/src/components/controls/index.ts
--- a/Engine/Plugins/VirtualProduction/RemoteControlWebInterface/WebApp/Client/src/components/controls/index.ts
+++ b/Engine/Plugins/VirtualProduction/RemoteControlWebInterface/WebApp/Client/src/components/controls/index.ts
@@ -1,6 +1,6 @@
 export * from './Button';
-export * from './Gauge';
-export * from './GaugesWrapper';
+export * from './Dial';
+export * from './DialsWrapper';
 export * from './Joystick';
 export * from './JoysticksWrapper';
 export * from './NewColorPicker';
@@ -10,3 +10,4 @@ export * from './Text';
 export * from './Toggle';
 export * from './ValueInput';
 export * from './VectorDrawer';
+export * from './Search';
\ No newline at end of file
diff --git a/Engine/Plugins/VirtualProduction/RemoteControlWebInterface/WebApp/Client/src/shared/index.ts b/Engine/Plugins/VirtualProduction/RemoteControlWebInterface/WebApp/Client/src/shared/index.ts
--- a/Engine/Plugins/VirtualProduction/RemoteControlWebInterface/WebApp/Client/src/shared/index.ts
+++ b/Engine/Plugins/VirtualProduction/RemoteControlWebInterface/WebApp/Client/src/shared/index.ts
@@ -128,7 +128,7 @@ export type IPayload = { [property: string]: PropertyValue | IPayload };
 export type IPayloads = { [preset: string]: IPayload };
 
 export enum WidgetTypes {
-  Gauge =           'Gauge',
+  Dial =            'Dial',
   Slider =          'Slider',
   Sliders =         'Sliders',
   ScaleSlider =     'Scale Slider',
@@ -206,7 +206,7 @@ export interface ICustomStackProperty {
   spaces?: number;
 
   // Dial only
-  mode?: 'ENDLESS' | 'RANGE' | 'ROTATION';
+  mode?: 'ENDLESS' | 'RANGE';
 }
 
 export interface ICustomStackItem {
diff --git a/Engine/Plugins/VirtualProduction/RemoteControlWebInterface/WebApp/Server/src/Notify.ts b/Engine/Plugins/VirtualProduction/RemoteControlWebInterface/WebApp/Server/src/Notify.ts
--- a/Engine/Plugins/VirtualProduction/RemoteControlWebInterface/WebApp/Server/src/Notify.ts
+++ b/Engine/Plugins/VirtualProduction/RemoteControlWebInterface/WebApp/Server/src/Notify.ts
@@ -30,8 +30,9 @@ export namespace Notify {
     io.emit(what, value);
   }
 
-  export function onViewChange(preset: string, view: IView) {
-    UnrealEngine.setView(preset, view);
+  export function onViewChange(preset: string, view: IView, supressUnrealNotification?: boolean) {
+    if (!supressUnrealNotification)
+      UnrealEngine.setView(preset, view);
     io.emit('view', preset, view);
   }
 
diff --git a/Engine/Plugins/VirtualProduction/RemoteControlWebInterface/WebApp/Server/src/UnrealEngine.ts b/Engine/Plugins/VirtualProduction/RemoteControlWebInterface/WebApp/Server/src/UnrealEngine.ts
--- a/Engine/Plugins/VirtualProduction/RemoteControlWebInterface/WebApp/Server/src/UnrealEngine.ts
+++ b/Engine/Plugins/VirtualProduction/RemoteControlWebInterface/WebApp/Server/src/UnrealEngine.ts
@@ -313,7 +313,7 @@ export namespace UnrealEngine {
       }
 
       views[preset] = view;
-      Notify.onViewChange(preset, view);
+      Notify.onViewChange(preset, view, true);
     } catch (error) {
       console.log('Failed to parse View of Preset', preset);
     }    
@@ -396,7 +396,6 @@ export namespace UnrealEngine {
       for (const group of Preset.Groups) {
         for (const property of group.ExposedProperties) {
            const value = await get<UnrealApi.PropertyValues>(`/remote/preset/${Preset.Name}/property/${property.ID}`);
-           console.log(value);
            setPayloadValueInternal(updatedPayloads, [Preset.Name, property.ID], value?.PropertyValues?.[0]?.PropertyValue);
         }
       }
@@ -527,11 +526,18 @@ export namespace UnrealEngine {
     }
   }
 
-  export async function setPresetPropertyMetadata(preset: string, property: string, metadata: string, value: string) {
+  export async function setPresetPropertyMetadata(preset: string, property: string, metadata: string, value: string, cb?: () => void) {
     try {
 
       const url = `/remote/preset/${preset}/property/${property}/metadata/${metadata}`;
       await put(url, { value });
+      const p = _.find(presets, p => p.Name === preset);
+      const prop = p?.Exposed[property];
+      if (prop) {
+        prop.Metadata[metadata] = value;
+        Notify.emit('presets', presets);
+        cb?.();
+      }
     } catch (err) {
       console.log(`Failed to rename a property`);
     }
diff --git a/Engine/Plugins/VirtualProduction/RemoteControlWebInterface/WebApp/package.json b/Engine/Plugins/VirtualProduction/RemoteControlWebInterface/WebApp/package.json
--- a/Engine/Plugins/VirtualProduction/RemoteControlWebInterface/WebApp/package.json
+++ b/Engine/Plugins/VirtualProduction/RemoteControlWebInterface/WebApp/package.json
@@ -1,6 +1,6 @@
 {
   "name": "conductor",
-  "version": "0.3.9",
+  "version": "0.3.10",
   "scripts": {
     "postinstall": "concurrently \"cd ./Client && npm install && cd ..\" \"cd ./Server && npm install && cd ..\"",
     "build": "node ./Scripts/build.js",
