RemoteControlWebInterface: Functions can't be added to UI

diff --git a/Engine/Plugins/VirtualProduction/RemoteControlWebInterface/WebApp/Client/src/reducers/api.ts b/Engine/Plugins/VirtualProduction/RemoteControlWebInterface/WebApp/Client/src/reducers/api.ts
--- a/Engine/Plugins/VirtualProduction/RemoteControlWebInterface/WebApp/Client/src/reducers/api.ts
+++ b/Engine/Plugins/VirtualProduction/RemoteControlWebInterface/WebApp/Client/src/reducers/api.ts
@@ -141,7 +141,7 @@ export const _api = {
     },
     execute: (func: string) => {
       const preset = _internal.getPreset();
-      _socket.emit('execute', null, preset, func);
+      _socket.emit('execute', preset, null, func);
     },
     asset: (asset: string, action: AssetAction, meta?: any) => {
       _socket.emit('asset', asset, action, meta);
diff --git a/Engine/Plugins/VirtualProduction/RemoteControlWebInterface/WebApp/Server/src/UnrealEngine.ts b/Engine/Plugins/VirtualProduction/RemoteControlWebInterface/WebApp/Server/src/UnrealEngine.ts
--- a/Engine/Plugins/VirtualProduction/RemoteControlWebInterface/WebApp/Server/src/UnrealEngine.ts
+++ b/Engine/Plugins/VirtualProduction/RemoteControlWebInterface/WebApp/Server/src/UnrealEngine.ts
@@ -280,10 +280,17 @@ export namespace UnrealEngine {
       Preset.Exposed = {};
 
       for (const Group of Preset.Groups) {
-        for (const Function of Group.ExposedFunctions)
-          Preset.Exposed[Function.ID] = Function;
-        for (const Property of Group.ExposedProperties)
+        for (const Property of Group.ExposedProperties) {
           Preset.Exposed[Property.ID] = Property;
+        }
+
+        for (const Function of Group.ExposedFunctions) {
+          if (!Function.Metadata)
+            Function.Metadata = {};
+
+          Function.Metadata.Widget = WidgetTypes.Button;
+          Preset.Exposed[Function.ID] = Function;
+        }
 
         Preset.ExposedProperties.push(...Group.ExposedProperties);
         Preset.ExposedFunctions.push(...Group.ExposedFunctions);
diff --git a/Engine/Plugins/VirtualProduction/RemoteControlWebInterface/WebApp/package.json b/Engine/Plugins/VirtualProduction/RemoteControlWebInterface/WebApp/package.json
--- a/Engine/Plugins/VirtualProduction/RemoteControlWebInterface/WebApp/package.json
+++ b/Engine/Plugins/VirtualProduction/RemoteControlWebInterface/WebApp/package.json
@@ -1,6 +1,6 @@
 {
   "name": "conductor",
-  "version": "0.3.10",
+  "version": "0.3.11",
   "scripts": {
     "postinstall": "concurrently \"cd ./Client && npm install && cd ..\" \"cd ./Server && npm install && cd ..\"",
     "build": "node ./Scripts/build.js",
