RemoteControlWebInterface: Cleanups/Reorder still WIP. A followup commit will follow due to renames (to match folder character case) - for some reason I cant push this in a single CL
+ Reorder top level tabs
+ mark already used properties in a tab (no style yet)
+ Fixed metadata set delay (e.g. setting min or max had some stuttering)
+ Set label override in Metadata.Description
+ Removed alert on preset change
+ Styled widgets panel
+ Added Label+Reset to ColorPicker & Dial
+ Joystick better label speed handles

diff --git a/Engine/Plugins/VirtualProduction/RemoteControlWebInterface/WebApp/Client/src/components/LayoutTabDrawer/index.ts b/Engine/Plugins/VirtualProduction/RemoteControlWebInterface/WebApp/Client/src/components/LayoutTabDrawer/index.ts
deleted file mode 100644
diff --git a/Engine/Plugins/VirtualProduction/RemoteControlWebInterface/WebApp/Client/src/components/Widgets/index.ts b/Engine/Plugins/VirtualProduction/RemoteControlWebInterface/WebApp/Client/src/components/Widgets/index.ts
deleted file mode 100644
diff --git a/Engine/Plugins/VirtualProduction/RemoteControlWebInterface/WebApp/Client/src/components/controls/index.ts b/Engine/Plugins/VirtualProduction/RemoteControlWebInterface/WebApp/Client/src/components/controls/index.ts
--- a/Engine/Plugins/VirtualProduction/RemoteControlWebInterface/WebApp/Client/src/components/controls/index.ts
+++ b/Engine/Plugins/VirtualProduction/RemoteControlWebInterface/WebApp/Client/src/components/controls/index.ts
@@ -1,9 +1,10 @@
 export * from './Button';
+export * from './ColorPicker';
 export * from './Dial';
 export * from './DialsWrapper';
 export * from './Joystick';
 export * from './JoysticksWrapper';
-export * from './NewColorPicker';
+export * from './SafeImage';
 export * from './Slider';
 export * from './SlidersWrapper';
 export * from './Text';
diff --git a/Engine/Plugins/VirtualProduction/RemoteControlWebInterface/WebApp/Client/src/components/index.ts b/Engine/Plugins/VirtualProduction/RemoteControlWebInterface/WebApp/Client/src/components/index.ts
--- a/Engine/Plugins/VirtualProduction/RemoteControlWebInterface/WebApp/Client/src/components/index.ts
+++ b/Engine/Plugins/VirtualProduction/RemoteControlWebInterface/WebApp/Client/src/components/index.ts
@@ -1,10 +1,7 @@
-export * from './AlertModal';
 export * from './ConnectionStatus';
-export * from './CustomStack';
+export * from './Stack';
 export * from './DraggableWidget';
-export * from './IconModal';
-export * from './PresetModal';
-export * from './RenameModal';
-export * from './SafeImage';
+export * from './PropertiesDrawer';
+export * from './modals';
 export * from './Screen';
 export * from './Tabs';
\ No newline at end of file
diff --git a/Engine/Plugins/VirtualProduction/RemoteControlWebInterface/WebApp/Client/src/components/modals/index.ts b/Engine/Plugins/VirtualProduction/RemoteControlWebInterface/WebApp/Client/src/components/modals/index.ts
new file mode 100644
--- /dev/null
+++ b/Engine/Plugins/VirtualProduction/RemoteControlWebInterface/WebApp/Client/src/components/modals/index.ts
@@ -0,0 +1,3 @@
+export * from './AlertModal';
+export * from './IconModal';
+export * from './RenameModal';
\ No newline at end of file
diff --git a/Engine/Plugins/VirtualProduction/RemoteControlWebInterface/WebApp/Client/src/components/widgetsRENAME/index.ts b/Engine/Plugins/VirtualProduction/RemoteControlWebInterface/WebApp/Client/src/components/widgetsRENAME/index.ts
new file mode 100644
--- /dev/null
+++ b/Engine/Plugins/VirtualProduction/RemoteControlWebInterface/WebApp/Client/src/components/widgetsRENAME/index.ts
@@ -0,0 +1,7 @@
+export * from './DropdownWidget';
+export * from './ImageSelectorWidget';
+export * from './SlidersWidget';
+export * from './SliderWidget';
+export * from './TabWidget';
+export * from './ToggleWidget';
+export * from './VectorWidget';
\ No newline at end of file
diff --git a/Engine/Plugins/VirtualProduction/RemoteControlWebInterface/WebApp/Client/src/reducers/api.ts b/Engine/Plugins/VirtualProduction/RemoteControlWebInterface/WebApp/Client/src/reducers/api.ts
--- a/Engine/Plugins/VirtualProduction/RemoteControlWebInterface/WebApp/Client/src/reducers/api.ts
+++ b/Engine/Plugins/VirtualProduction/RemoteControlWebInterface/WebApp/Client/src/reducers/api.ts
@@ -66,10 +66,6 @@ function _initialize(dispatch: Dispatch, getState: () => { api: ApiState }) {
 
 type IRequestCallback = Function | string | undefined;
 
-interface IRequestOptions {
-  blob?: boolean;
-}
-
 async function _request(method: string, url: string, body: string | object | undefined, callback: IRequestCallback): Promise<any> {
   const request: RequestInit = { method, mode: 'cors', redirect: 'follow', headers: {} };
   if (body instanceof FormData || typeof(body) === 'string') {
@@ -150,9 +146,9 @@ export const _api = {
       const preset = _internal.getPreset();
       _socket.emit('rename', preset, type, property, label, callback);
     },
-    metadata: (property: string, meta: string, value: string, callback: () => void) => {
+    metadata: (property: string, meta: string, value: string) => {
       const preset = _internal.getPreset();
-      _socket.emit('metadata', preset, property, meta, value, callback);
+      _socket.emit('metadata', preset, property, meta, value);
     }
   },
   actor: {
diff --git a/Engine/Plugins/VirtualProduction/RemoteControlWebInterface/WebApp/Client/src/shared/index.ts b/Engine/Plugins/VirtualProduction/RemoteControlWebInterface/WebApp/Client/src/shared/index.ts
--- a/Engine/Plugins/VirtualProduction/RemoteControlWebInterface/WebApp/Client/src/shared/index.ts
+++ b/Engine/Plugins/VirtualProduction/RemoteControlWebInterface/WebApp/Client/src/shared/index.ts
@@ -72,6 +72,9 @@ export interface IExposedFunction {
   DisplayName: string;
   UnderlyingFunction: IFunction;
   Metadata: { [key: string]: string };
+
+  //Added
+  Type: PropertyType;
 }
 
 export interface IProperty {
@@ -87,6 +90,9 @@ export interface IExposedProperty {
   Metadata: Record<string, string>;
   Widget: WidgetType;
   UnderlyingProperty: IProperty;
+
+  //Added
+  Type: PropertyType;
 }
 
 export interface IActor {
@@ -142,7 +148,8 @@ export enum WidgetTypes {
   ImageSelector =   'Image Selector',
   Vector =          'Vector',
   Spacer =          'Spacer',
-
+  Tabs =            'Tabs',
+  
   Level =           'Level',
   Sequence =        'Sequence',
 }
@@ -150,10 +157,11 @@ export enum WidgetTypes {
 export type WidgetType = keyof typeof WidgetTypes | string;
 
 
-export type IWidgetMeta = { [key: string]: any } & {
+export type IWidgetMeta = {
+  Description?: string;
   Min?: number;
   Max?: number;
-};
+} & { [key: string]: any };
 
 
 export enum IPanelType {
@@ -174,9 +182,13 @@ export enum TabLayout {
   Screen =   'Screen',
 }
 
+export enum ScreenType {
+  Snapshot =      'Snapshot',
+}
+
 export interface IScreen {
-  type: 'Snapshot';
-  data: any;
+  type: ScreenType;
+  data?: any;
 }
 
 export interface IDropdownOption {
@@ -189,24 +201,18 @@ export interface ICustomStackProperty {
   property: string;
   propertyType: PropertyType;
   widget: WidgetType;
-  label?: string;
 
-  // Sliders & Vector only
-  lock?: boolean;
+  // Label only
+  label?: string;
 
   // Vector only
   widgets?: WidgetType[];
-  speedMin?: number;
-  speedMax?: number;
 
   // Dropdown only
   options?: IDropdownOption[];
 
   // Space only
   spaces?: number;
-
-  // Dial only
-  mode?: 'ENDLESS' | 'RANGE';
 }
 
 export interface ICustomStackItem {
diff --git a/Engine/Plugins/VirtualProduction/RemoteControlWebInterface/WebApp/Client/src/utilities/WidgetUtilities.ts b/Engine/Plugins/VirtualProduction/RemoteControlWebInterface/WebApp/Client/src/utilities/WidgetUtilities.ts
--- a/Engine/Plugins/VirtualProduction/RemoteControlWebInterface/WebApp/Client/src/utilities/WidgetUtilities.ts
+++ b/Engine/Plugins/VirtualProduction/RemoteControlWebInterface/WebApp/Client/src/utilities/WidgetUtilities.ts
@@ -104,9 +104,11 @@ export class WidgetUtilities {
       case PropertyType.Vector:
       case PropertyType.Vector2D:
       case PropertyType.Vector4:
-      case PropertyType.Rotator:
       case PropertyType.LinearColor:
         return 3;
+
+      case PropertyType.Rotator:
+        return 2;
     }
 
     return 0;
diff --git a/Engine/Plugins/VirtualProduction/RemoteControlWebInterface/WebApp/Server/src/UnrealEngine.ts b/Engine/Plugins/VirtualProduction/RemoteControlWebInterface/WebApp/Server/src/UnrealEngine.ts
--- a/Engine/Plugins/VirtualProduction/RemoteControlWebInterface/WebApp/Server/src/UnrealEngine.ts
+++ b/Engine/Plugins/VirtualProduction/RemoteControlWebInterface/WebApp/Server/src/UnrealEngine.ts
@@ -281,6 +281,7 @@ export namespace UnrealEngine {
 
       for (const Group of Preset.Groups) {
         for (const Property of Group.ExposedProperties) {
+          Property.Type = Property.UnderlyingProperty.Type;
           Preset.Exposed[Property.ID] = Property;
         }
 
@@ -288,6 +289,7 @@ export namespace UnrealEngine {
           if (!Function.Metadata)
             Function.Metadata = {};
 
+          Function.Type = PropertyType.Function;
           Function.Metadata.Widget = WidgetTypes.Button;
           Preset.Exposed[Function.ID] = Function;
         }
@@ -533,7 +535,7 @@ export namespace UnrealEngine {
     }
   }
 
-  export async function setPresetPropertyMetadata(preset: string, property: string, metadata: string, value: string, cb?: () => void) {
+  export async function setPresetPropertyMetadata(preset: string, property: string, metadata: string, value: string) {
     try {
 
       const url = `/remote/preset/${preset}/property/${property}/metadata/${metadata}`;
@@ -543,7 +545,6 @@ export namespace UnrealEngine {
       if (prop) {
         prop.Metadata[metadata] = value;
         Notify.emit('presets', presets);
-        cb?.();
       }
     } catch (err) {
       console.log(`Failed to rename a property`);
diff --git a/Engine/Plugins/VirtualProduction/RemoteControlWebInterface/WebApp/package.json b/Engine/Plugins/VirtualProduction/RemoteControlWebInterface/WebApp/package.json
--- a/Engine/Plugins/VirtualProduction/RemoteControlWebInterface/WebApp/package.json
+++ b/Engine/Plugins/VirtualProduction/RemoteControlWebInterface/WebApp/package.json
@@ -1,6 +1,6 @@
 {
   "name": "conductor",
-  "version": "0.3.11",
+  "version": "0.3.12",
   "scripts": {
     "postinstall": "concurrently \"cd ./Client && npm install && cd ..\" \"cd ./Server && npm install && cd ..\"",
     "build": "node ./Scripts/build.js",
