RemoteControlWebInterface: Perforce won, I can't change folder case, reverting this to match original

diff --git a/Engine/Plugins/VirtualProduction/RemoteControlWebInterface/WebApp/Client/src/components/Screens/index.ts b/Engine/Plugins/VirtualProduction/RemoteControlWebInterface/WebApp/Client/src/components/Screens/index.ts
new file mode 100644
--- /dev/null
+++ b/Engine/Plugins/VirtualProduction/RemoteControlWebInterface/WebApp/Client/src/components/Screens/index.ts
@@ -0,0 +1 @@
+export * from './Snapshot';
diff --git a/Engine/Plugins/VirtualProduction/RemoteControlWebInterface/WebApp/Client/src/components/Widgets/index.ts b/Engine/Plugins/VirtualProduction/RemoteControlWebInterface/WebApp/Client/src/components/Widgets/index.ts
new file mode 100644
--- /dev/null
+++ b/Engine/Plugins/VirtualProduction/RemoteControlWebInterface/WebApp/Client/src/components/Widgets/index.ts
@@ -0,0 +1,7 @@
+export * from './DropdownWidget';
+export * from './ImageSelectorWidget';
+export * from './SlidersWidget';
+export * from './SliderWidget';
+export * from './TabWidget';
+export * from './ToggleWidget';
+export * from './VectorWidget';
\ No newline at end of file
