RemoteControlWebInterface: Bring back LevelSnapshot UI
+ Trying to improve first load time read preset+values per preset

diff --git a/Engine/Plugins/VirtualProduction/RemoteControlWebInterface/WebApp/Server/src/UnrealEngine.ts b/Engine/Plugins/VirtualProduction/RemoteControlWebInterface/WebApp/Server/src/UnrealEngine.ts
--- a/Engine/Plugins/VirtualProduction/RemoteControlWebInterface/WebApp/Server/src/UnrealEngine.ts
+++ b/Engine/Plugins/VirtualProduction/RemoteControlWebInterface/WebApp/Server/src/UnrealEngine.ts
@@ -118,8 +118,7 @@ export namespace UnrealEngine {
 
   async function refresh() {
     try {
-      await pullPresets();
-      await pullValues();
+      await pullPresets(true);
     } catch (error) {
     }
   }
@@ -241,22 +240,34 @@ export namespace UnrealEngine {
     return presets;
   }
 
-  async function pullPresets(): Promise<void> {
+  async function pullPresets(pullValues?: boolean): Promise<void> {
     try {
-      const all = [];
+      const allPresets = [];
+      const allPayloads: IPayloads = {};
+  
       const { Presets } = await get<UnrealApi.Presets>('/remote/presets');
       
       for (const p of Presets ?? []) {
-        const preset = await pullPreset(p.Name);
-        if (preset)
-          all.push(preset);
+        const Preset = await pullPreset(p.Name);
+        if (!Preset)
+          continue;
+
+        allPresets.push(Preset);
+        if (pullValues)
+          allPayloads[Preset.Name] = await pullPresetValues(Preset);
       }
 
-      const compact =  _.compact(all);
+      const compact =  _.compact(allPresets);
       if (!equal(presets, compact)) {
         presets = compact;
         Notify.emit('presets', presets);
       }
+
+      if (pullValues && !equal(payloads, allPayloads)) {
+        payloads = allPayloads;
+        Notify.emit('payloads', payloads );
+      }
+
     } catch (error) {
       console.log('Failed to pull presets data');
     }
@@ -399,21 +410,14 @@ export namespace UnrealEngine {
     return false;
   }
 
-  async function pullValues(): Promise<void> {
-    const updatedPayloads = {};
-    for (const Preset of presets) {
-      for (const group of Preset.Groups) {
-        for (const property of group.ExposedProperties) {
-           const value = await get<UnrealApi.PropertyValues>(`/remote/preset/${Preset.Name}/property/${property.ID}`);
-           setPayloadValueInternal(updatedPayloads, [Preset.Name, property.ID], value?.PropertyValues?.[0]?.PropertyValue);
-        }
-      }
+  async function pullPresetValues(Preset: IPreset): Promise<IPayload> {
+    const updatedPayloads: IPayloads = {};
+    for (const property of Preset.ExposedProperties) {
+      const value = await get<UnrealApi.PropertyValues>(`/remote/preset/${Preset.Name}/property/${property.ID}`);
+      setPayloadValueInternal(updatedPayloads, [Preset.Name, property.ID], value?.PropertyValues?.[0]?.PropertyValue);
     }
 
-    if (!equal(payloads, updatedPayloads)) {
-      payloads = updatedPayloads;
-      Notify.emit('payloads', payloads );
-    }
+    return updatedPayloads[Preset.Name];
   }
 
   export async function getPayload(preset: string): Promise<IPayload> {
@@ -619,6 +623,7 @@ export namespace UnrealEngine {
       case 'PUT':
         if (body)
           return put(url, body);
+        break;
     }
 
     return Promise.resolve({});
diff --git a/Engine/Plugins/VirtualProduction/RemoteControlWebInterface/WebApp/package.json b/Engine/Plugins/VirtualProduction/RemoteControlWebInterface/WebApp/package.json
--- a/Engine/Plugins/VirtualProduction/RemoteControlWebInterface/WebApp/package.json
+++ b/Engine/Plugins/VirtualProduction/RemoteControlWebInterface/WebApp/package.json
@@ -1,6 +1,6 @@
 {
   "name": "conductor",
-  "version": "0.3.13",
+  "version": "0.3.14",
   "scripts": {
     "postinstall": "concurrently \"cd ./Client && npm install && cd ..\" \"cd ./Server && npm install && cd ..\"",
     "build": "node ./Scripts/build.js",
