Fixing case of //UE5/Main/Engine/Plugins/VirtualProduction/RemoteControlWebInterface/WebApp/Client/Public/... (2/2)

#rnx

diff --git a/Engine/Plugins/VirtualProduction/RemoteControlWebInterface/WebApp/Client/Public/index.html b/Engine/Plugins/VirtualProduction/RemoteControlWebInterface/WebApp/Client/Public/index.html
new file mode 100644
--- /dev/null
+++ b/Engine/Plugins/VirtualProduction/RemoteControlWebInterface/WebApp/Client/Public/index.html
@@ -0,0 +1,13 @@
+<html lang="en">
+<head>
+  <title>Conductor</title>
+  <link rel="apple-touch-icon" sizes="180x180" href="/images/favicon/apple-touch-icon.png">
+  <link rel="icon" type="image/png" sizes="32x32" href="/images/favicon/32x32.png">
+  <link rel="icon" type="image/png" sizes="16x16" href="/images/favicon/16x16.png">
+  <link rel="manifest" href="/manifest.json">
+  <meta name="viewport" content="user-scalable=no" />
+</head>
+<body>
+  <div id="root"></div>
+</body>
+</html>
\ No newline at end of file
diff --git a/Engine/Plugins/VirtualProduction/RemoteControlWebInterface/WebApp/Client/Public/manifest.json b/Engine/Plugins/VirtualProduction/RemoteControlWebInterface/WebApp/Client/Public/manifest.json
new file mode 100644
--- /dev/null
+++ b/Engine/Plugins/VirtualProduction/RemoteControlWebInterface/WebApp/Client/Public/manifest.json
@@ -0,0 +1,26 @@
+{
+  "name": "Unreal Engine",
+  "short_name": "Unreal Engine",
+  "icons": [
+    {
+      "src": "/images/favicon/android-chrome-192x192.png?v=xQdXpAz40w",
+      "sizes": "192x192",
+      "type": "image/png"
+    },
+    {
+      "src": "/images/favicon/android-chrome-256x256.png?v=xQdXpAz40w",
+      "sizes": "256x256",
+      "type": "image/png"
+    },
+    {
+      "src": "/images/favicon/android-chrome-512x512.png?v=xQdXpAz40w",
+      "sizes": "512x512",
+      "type": "image/png"
+    }
+  ],
+  "theme_color": "#ffffff",
+  "background_color": "#ffffff",
+  "display": "browser",
+  "scope": "/",
+  "start_url": "/"
+}
\ No newline at end of file
