RemoteControlWebInterface: Adding Alpha to Color Pickers (FColor & FLinearColor)
Fix - When more than 1 instance opens and all runs on same port, only last instance should work
ColorPickers - Fix values stuck on certian values

diff --git a/Engine/Plugins/VirtualProduction/RemoteControlWebInterface/WebApp/Client/src/shared/index.ts b/Engine/Plugins/VirtualProduction/RemoteControlWebInterface/WebApp/Client/src/shared/index.ts
--- a/Engine/Plugins/VirtualProduction/RemoteControlWebInterface/WebApp/Client/src/shared/index.ts
+++ b/Engine/Plugins/VirtualProduction/RemoteControlWebInterface/WebApp/Client/src/shared/index.ts
@@ -26,6 +26,7 @@ export interface ColorProperty {
   R: number;
   G: number;
   B: number;
+  A?: number;
 }
 
 export interface VectorProperty {
@@ -50,6 +51,7 @@ export type IHsvColor = {
   h: number;
   s: number;
   v: number;
+  a?: number;
 };
 
 export interface IFunctionParameter {
diff --git a/Engine/Plugins/VirtualProduction/RemoteControlWebInterface/WebApp/Client/src/utilities/WidgetUtilities.ts b/Engine/Plugins/VirtualProduction/RemoteControlWebInterface/WebApp/Client/src/utilities/WidgetUtilities.ts
--- a/Engine/Plugins/VirtualProduction/RemoteControlWebInterface/WebApp/Client/src/utilities/WidgetUtilities.ts
+++ b/Engine/Plugins/VirtualProduction/RemoteControlWebInterface/WebApp/Client/src/utilities/WidgetUtilities.ts
@@ -166,11 +166,16 @@ export class WidgetUtilities {
     const vector = color as VectorProperty;
 
     max = max || 1;
-    return {
+    const rgbColor = {
       R: (vector?.X ?? rgb?.R) * 255 / max,
       G: (vector?.Y ?? rgb?.G) * 255 / max,
       B: (vector?.Z ?? rgb?.B) * 255 / max,
-    };
+    } as ColorProperty;
+
+    if (rgb?.A)
+      rgbColor.A = rgb.A;
+
+    return rgbColor;
   }
 
   static rgbToColor = (rgb: ColorProperty, max: number, vector = true): ColorProperty | VectorProperty => {
@@ -181,7 +186,11 @@ export class WidgetUtilities {
     if (vector)
       return { X: a, Y: b, Z: c };
 
-    return { R: a, G: b, B: c };
+    const color = { R: a, G: b, B: c } as ColorProperty;
+    if (rgb.A)
+      color.A = rgb.A;
+
+    return color;
   }
 
   static rgb2Hsv(color: ColorProperty): IHsvColor {
@@ -214,10 +223,14 @@ export class WidgetUtilities {
       h /= 6;
     }
 
-    return { h, s, v };
+    const hsvColor = { h, s, v } as IHsvColor;
+    if (color?.A)
+      hsvColor.a = color.A;
+
+    return hsvColor;
   }
 
-  static hsv2rgb({ h, s, v }: IHsvColor): ColorProperty {
+  static hsv2rgb({ h, s, v, a }: IHsvColor): ColorProperty {
     let r: number, g: number, b: number;
 
     let i = Math.floor(h * 6);
@@ -244,6 +257,10 @@ export class WidgetUtilities {
     if (i % 6 === 5)
       ([r, g, b] = [v, p, q]);
 
-    return { R: r * 255, G: g * 255, B: b * 255 };
+    const color = { R: r * 255, G: g * 255, B: b * 255 } as ColorProperty;
+    if (a)
+      color.A = a;
+
+    return color;
   }
 }
\ No newline at end of file
diff --git a/Engine/Plugins/VirtualProduction/RemoteControlWebInterface/WebApp/Server/src/Api.ts b/Engine/Plugins/VirtualProduction/RemoteControlWebInterface/WebApp/Server/src/Api.ts
--- a/Engine/Plugins/VirtualProduction/RemoteControlWebInterface/WebApp/Server/src/Api.ts
+++ b/Engine/Plugins/VirtualProduction/RemoteControlWebInterface/WebApp/Server/src/Api.ts
@@ -66,7 +66,7 @@ export class Api {
   @Get('shutdown')
   private shutdown(req: Request, res: Response) {
     res.send({ message: 'ok' });
-    setTimeout(() => process.exit(0), 1000);
+    setImmediate(() => process.exit(0));
   }
 
   protected async send(res: Response, promise: Promise<any>, options?: ISendOptions): Promise<any> {
diff --git a/Engine/Plugins/VirtualProduction/RemoteControlWebInterface/WebApp/Server/src/App.ts b/Engine/Plugins/VirtualProduction/RemoteControlWebInterface/WebApp/Server/src/App.ts
--- a/Engine/Plugins/VirtualProduction/RemoteControlWebInterface/WebApp/Server/src/App.ts
+++ b/Engine/Plugins/VirtualProduction/RemoteControlWebInterface/WebApp/Server/src/App.ts
@@ -29,6 +29,7 @@ export class App extends Server {
       // Trying to kill a zombie process
       await request.get(`http://127.0.0.1:${Program.port}/api/shutdown`)
                     .timeout(100)
+                    .then(() => new Promise(resolve => setTimeout(resolve, 100)))
                     .catch(() => {});
 
       await this.startServer();
diff --git a/Engine/Plugins/VirtualProduction/RemoteControlWebInterface/WebApp/_package.json b/Engine/Plugins/VirtualProduction/RemoteControlWebInterface/WebApp/_package.json
--- a/Engine/Plugins/VirtualProduction/RemoteControlWebInterface/WebApp/_package.json
+++ b/Engine/Plugins/VirtualProduction/RemoteControlWebInterface/WebApp/_package.json
@@ -1,6 +1,6 @@
 {
   "name": "conductor",
-  "version": "0.3.54",
+  "version": "0.3.55",
   "scripts": {
     "postinstall": "concurrently \"cd ./Client && npm install && cd ..\" \"cd ./Server && npm install && cd ..\"",
     "build": "node ./Scripts/build.js",
