RemoteControlWebInterface: - Adding Dials widget (vector)
- Export a list into multiple tabs
- Play button select area should be whole button

diff --git a/Engine/Plugins/VirtualProduction/RemoteControlWebInterface/WebApp/Client/src/shared/index.ts b/Engine/Plugins/VirtualProduction/RemoteControlWebInterface/WebApp/Client/src/shared/index.ts
--- a/Engine/Plugins/VirtualProduction/RemoteControlWebInterface/WebApp/Client/src/shared/index.ts
+++ b/Engine/Plugins/VirtualProduction/RemoteControlWebInterface/WebApp/Client/src/shared/index.ts
@@ -45,8 +45,6 @@ export type PropertyValue = boolean | number | string | VectorProperty | Rotator
 
 export type JoystickValue = { [key: string]: number };
 
-export type AxisInfo = [string, number];
-
 export type IHsvColor = {
   h: number;
   s: number;
@@ -143,6 +141,7 @@ export type IPayloads = { [preset: string]: IPayload };
 
 export enum WidgetTypes {
   Dial =            'Dial',
+  Dials =           'Dials',
   Slider =          'Slider',
   Sliders =         'Sliders',
   ScaleSlider =     'Scale Slider',
diff --git a/Engine/Plugins/VirtualProduction/RemoteControlWebInterface/WebApp/Client/src/utilities/WidgetUtilities.ts b/Engine/Plugins/VirtualProduction/RemoteControlWebInterface/WebApp/Client/src/utilities/WidgetUtilities.ts
--- a/Engine/Plugins/VirtualProduction/RemoteControlWebInterface/WebApp/Client/src/utilities/WidgetUtilities.ts
+++ b/Engine/Plugins/VirtualProduction/RemoteControlWebInterface/WebApp/Client/src/utilities/WidgetUtilities.ts
@@ -144,7 +144,7 @@ export class WidgetUtilities {
 
       case PropertyType.Vector:
       case PropertyType.Vector2D:
-        return [WidgetTypes.Vector, WidgetTypes.Joystick, WidgetTypes.Sliders];
+        return [WidgetTypes.Vector, WidgetTypes.Joystick, WidgetTypes.Sliders, WidgetTypes.Dials];
 
       case PropertyType.Vector4:
       case PropertyType.LinearColor:
@@ -152,7 +152,7 @@ export class WidgetUtilities {
         return [WidgetTypes.ColorPicker, WidgetTypes.MiniColorPicker];
 
       case PropertyType.Rotator:
-        return [WidgetTypes.Vector, WidgetTypes.Sliders];
+        return [WidgetTypes.Vector, WidgetTypes.Sliders, WidgetTypes.Dials];
     }
 
     if (propertyType && (propertyType.startsWith('TEnum') || propertyType.startsWith('E')))
diff --git a/Engine/Plugins/VirtualProduction/RemoteControlWebInterface/WebApp/_package.json b/Engine/Plugins/VirtualProduction/RemoteControlWebInterface/WebApp/_package.json
--- a/Engine/Plugins/VirtualProduction/RemoteControlWebInterface/WebApp/_package.json
+++ b/Engine/Plugins/VirtualProduction/RemoteControlWebInterface/WebApp/_package.json
@@ -1,6 +1,6 @@
 {
   "name": "conductor",
-  "version": "0.3.57",
+  "version": "0.3.58",
   "scripts": {
     "postinstall": "concurrently \"cd ./Client && npm install && cd ..\" \"cd ./Server && npm install && cd ..\"",
     "build": "node ./Scripts/build.js",
