RemoteControlWebInterface - webapp is not populating level snapshots Filter dropdown

diff --git a/Engine/Plugins/VirtualProduction/RemoteControlWebInterface/WebApp/Client/src/reducers/api.ts b/Engine/Plugins/VirtualProduction/RemoteControlWebInterface/WebApp/Client/src/reducers/api.ts
--- a/Engine/Plugins/VirtualProduction/RemoteControlWebInterface/WebApp/Client/src/reducers/api.ts
+++ b/Engine/Plugins/VirtualProduction/RemoteControlWebInterface/WebApp/Client/src/reducers/api.ts
@@ -154,19 +154,8 @@ export const _api = {
     },
   },
   assets: {
-    search: (q: string, types: string[], prefix: string, count: number = 50): Promise<IAsset[]> => {
-      const args = {
-        q,
-        prefix,
-        count,
-        types: types.join(','),
-      };
-      
-      let url = '/api/assets/search?';
-      for (const arg in args)
-        url += `${arg}=${encodeURIComponent(args[arg])}&`;
-      
-      return _get(url);
+    search: (q: string, types: string[], prefix: string, filterArgs = {}, count: number = 50): Promise<IAsset[]> => {
+      return new Promise(resolve => _socket.emit('search', q, types, prefix, filterArgs, count, resolve));
     },
     thumbnailUrl: (asset: string) => `${_host}/api/thumbnail?asset=${asset}`,
   },
diff --git a/Engine/Plugins/VirtualProduction/RemoteControlWebInterface/WebApp/Server/src/Api.ts b/Engine/Plugins/VirtualProduction/RemoteControlWebInterface/WebApp/Server/src/Api.ts
--- a/Engine/Plugins/VirtualProduction/RemoteControlWebInterface/WebApp/Server/src/Api.ts
+++ b/Engine/Plugins/VirtualProduction/RemoteControlWebInterface/WebApp/Server/src/Api.ts
@@ -44,15 +44,6 @@ export class Api {
     this.send(res, UnrealEngine.getView(req.query.preset?.toString()));
   }
 
-  @Get('assets/search')
-  private search(req: Request, res: Response) {
-    const prefix = req.query.prefix?.toString() || '/Game';
-    const query = req.query.q?.toString() ?? '';
-    const count = parseInt(req.query.count?.toString()) || 50;
-    const types = req.query.types?.toString().split(',') ?? [];
-    this.send(res, UnrealEngine.search(query, types, prefix, count));
-  }
-
   @Put('proxy')
   private proxy(req: Request, res: Response) {
     this.send(res, UnrealEngine.proxy(req.body?.method, req.body?.url, req.body?.body));
diff --git a/Engine/Plugins/VirtualProduction/RemoteControlWebInterface/WebApp/Server/src/Notify.ts b/Engine/Plugins/VirtualProduction/RemoteControlWebInterface/WebApp/Server/src/Notify.ts
--- a/Engine/Plugins/VirtualProduction/RemoteControlWebInterface/WebApp/Server/src/Notify.ts
+++ b/Engine/Plugins/VirtualProduction/RemoteControlWebInterface/WebApp/Server/src/Notify.ts
@@ -16,7 +16,8 @@ export namespace Notify {
         .on('value', UnrealEngine.setPayloadValue)
         .on('execute', UnrealEngine.executeFunction)
         .on('metadata', UnrealEngine.setPresetPropertyMetadata)
-        .on('rebind', UnrealEngine.rebindProperties);
+        .on('rebind', UnrealEngine.rebindProperties)
+        .on('search', UnrealEngine.search);
 
       if (UnrealEngine.isConnected())
         socket.emit('connected', true);
diff --git a/Engine/Plugins/VirtualProduction/RemoteControlWebInterface/WebApp/Server/src/UnrealEngine.ts b/Engine/Plugins/VirtualProduction/RemoteControlWebInterface/WebApp/Server/src/UnrealEngine.ts
--- a/Engine/Plugins/VirtualProduction/RemoteControlWebInterface/WebApp/Server/src/UnrealEngine.ts
+++ b/Engine/Plugins/VirtualProduction/RemoteControlWebInterface/WebApp/Server/src/UnrealEngine.ts
@@ -708,18 +708,21 @@ export namespace UnrealEngine {
     await put(`/remote/preset/${preset}/metadata/view`, { Value });
   }
 
-  export async function search(query: string, types: string[], prefix: string, count: number): Promise<IAsset[]> {
+  export async function search(query: string, types: string[], prefix: string, filterArgs: any, count: number, callback: (assets: IAsset[]) => void): Promise<void> {
     const ret = await put<UnrealApi.Assets>('/remote/search/assets', {
       Query: query,
       Limit: count,
+      
       Filter: {
         ClassNames: types,
         PackagePaths: [prefix],
-        RecursivePaths: true
+        RecursivePaths: true,
+        RecursiveClasses: true,
+        ...filterArgs,
       }
     });
 
-    return ret.Assets;
+    callback?.(ret.Assets);
   }
 
   export function proxy(method: 'GET' | 'PUT', url: string, body?: any): Promise<any> {
diff --git a/Engine/Plugins/VirtualProduction/RemoteControlWebInterface/WebApp/_package.json b/Engine/Plugins/VirtualProduction/RemoteControlWebInterface/WebApp/_package.json
--- a/Engine/Plugins/VirtualProduction/RemoteControlWebInterface/WebApp/_package.json
+++ b/Engine/Plugins/VirtualProduction/RemoteControlWebInterface/WebApp/_package.json
@@ -1,6 +1,6 @@
 {
   "name": "conductor",
-  "version": "0.3.58",
+  "version": "0.3.60",
   "scripts": {
     "postinstall": "concurrently \"cd ./Client && npm install && cd ..\" \"cd ./Server && npm install && cd ..\"",
     "build": "node ./Scripts/build.js",
