RemoteControlWebInterface - WebApp is missing imporved logging & highlights
Merge from Marmot427 CL 18827335

diff --git a/Engine/Plugins/VirtualProduction/RemoteControlWebInterface/WebApp/Server/src/LogServer.ts b/Engine/Plugins/VirtualProduction/RemoteControlWebInterface/WebApp/Server/src/LogServer.ts
--- a/Engine/Plugins/VirtualProduction/RemoteControlWebInterface/WebApp/Server/src/LogServer.ts
+++ b/Engine/Plugins/VirtualProduction/RemoteControlWebInterface/WebApp/Server/src/LogServer.ts
@@ -4,22 +4,36 @@ import fs from 'fs-extra';
 import { Program } from './';
 
 
+// if command takes more than durationThreashold (in ms) log it in the "Highlights" file
+const durationThreashold = 1000;
+
+interface IRequestDuration {
+  RequestId: number;
+  payload: any;
+  startTime: Date;
+}
+
 export namespace LogServer {
 
   var server: ws.Server;
+  var loopback: ws;
   var buffer: string = '';
+  var highlights: string = '';
+  var durations: Record<number, IRequestDuration> = {};
   var startTime = new Date().toISOString().replace(/\:/g, '_');
   var flushTimer: NodeJS.Timeout;
 
   export async function initialize(): Promise<void> {
     const port = Program.dev ? 7002 : Program.port + 2;
     server = new ws.Server({ port });
+    loopback = new ws(`ws://localhost:${port}`);
+
     await fs.ensureDir('./Logs');
     server.on('connection', (socket: ws) => {
       socket.on('message', onMessage);
     });
 
-    flushTimer = setInterval(flusLog, 2 * 60 * 1000);
+    flushTimer = setInterval(flushLogs, 30 * 1000);
   }
 
   export async function shutdown(): Promise<void> {
@@ -41,14 +55,39 @@ export namespace LogServer {
     }
   }
 
-  export function log(payload: Record<string, any>): void {
+  function log(payload: Record<string, any>, lastRequestLog?: boolean): void {
     const Timestamp = new Date().toISOString();
     buffer += JSON.stringify({ Timestamp, ...payload }) + '\r\n';
+
+    const { RequestId } = payload;
+    let info = durations[RequestId];
+    if (!info) {
+      info = { RequestId, payload, startTime: new Date() };
+      durations[payload.RequestId] = info;
+    }
+
+    if (lastRequestLog) {
+      const duration = new Date().valueOf() - info.startTime.valueOf();
+      if (duration > durationThreashold) {
+        highlights += `RequestId ${RequestId}, took ${duration}ms\r\n`;
+        highlights += JSON.stringify(info.payload) + '\r\n\r\n';
+      }
+
+      delete durations[RequestId];
+    }
+  }
+
+  export function logLoopback(payload: Record<string, any>, lastRequestLog?: boolean): void {
+    if (loopback?.readyState === ws.OPEN)
+      log(payload, lastRequestLog);
   }
   
-  function flusLog() {
+  function flushLogs() {
     if (flushLogFile(buffer, 'all'))
       buffer = '';
+
+    if (flushLogFile(highlights, 'highlights'))
+      highlights = '';
   }
 
   function flushLogFile(content: string, type: string) {
diff --git a/Engine/Plugins/VirtualProduction/RemoteControlWebInterface/WebApp/Server/src/UnrealEngine.ts b/Engine/Plugins/VirtualProduction/RemoteControlWebInterface/WebApp/Server/src/UnrealEngine.ts
--- a/Engine/Plugins/VirtualProduction/RemoteControlWebInterface/WebApp/Server/src/UnrealEngine.ts
+++ b/Engine/Plugins/VirtualProduction/RemoteControlWebInterface/WebApp/Server/src/UnrealEngine.ts
@@ -181,7 +181,7 @@ export namespace UnrealEngine {
           return;
 
         if (Program.logger)
-          LogServer.log({ RequestId: message.RequestId, Stage: 'Done' });
+          LogServer.logLoopback({ RequestId: message.RequestId, Stage: 'WebApp Done' }, true);
 
         delete pendings[message.RequestId];
         promise?.(message.ResponseBody);
@@ -305,7 +305,7 @@ export namespace UnrealEngine {
     const RequestId = httpRequest++;
     const payload = { RequestId, Verb, URL, Body };
     if (Program.logger)
-      LogServer.log(payload);
+      LogServer.logLoopback(payload);
 
     send('http', payload);
     if (!wantAnswer)
@@ -608,6 +608,19 @@ export namespace UnrealEngine {
         body.ResetToDefault = true; 
       }
 
+      if (Program.logger) {
+        const presetObj = presets[preset];
+        const log = {
+          RequestId: httpRequest,
+          Stage: 'Set Property',
+          Preset: presetObj?.Name,
+          Property: presetObj?.Exposed[property]?.DisplayName,
+          Value: value
+        };
+
+        LogServer.logLoopback(log);
+      }
+
       await put(`/remote/preset/${preset}/property/${property}`, body);
 
       if (value === null) {
diff --git a/Engine/Plugins/VirtualProduction/RemoteControlWebInterface/WebApp/_package.json b/Engine/Plugins/VirtualProduction/RemoteControlWebInterface/WebApp/_package.json
--- a/Engine/Plugins/VirtualProduction/RemoteControlWebInterface/WebApp/_package.json
+++ b/Engine/Plugins/VirtualProduction/RemoteControlWebInterface/WebApp/_package.json
@@ -1,6 +1,6 @@
 {
   "name": "conductor",
-  "version": "0.3.61",
+  "version": "0.3.62",
   "scripts": {
     "postinstall": "concurrently \"cd ./Client && npm install && cd ..\" \"cd ./Server && npm install && cd ..\"",
     "build": "node ./Scripts/build.js",
