Add: Loginpage Throwback

diff --git a/Engine/Plugins/VirtualProduction/RemoteControl/Source/WebRemoteControl/Private/WebRemoteControl.cpp b/Engine/Plugins/VirtualProduction/RemoteControl/Source/WebRemoteControl/Private/WebRemoteControl.cpp
--- a/Engine/Plugins/VirtualProduction/RemoteControl/Source/WebRemoteControl/Private/WebRemoteControl.cpp
+++ b/Engine/Plugins/VirtualProduction/RemoteControl/Source/WebRemoteControl/Private/WebRemoteControl.cpp
@@ -273,14 +273,17 @@ void FWebRemoteControlModule::StartupModule()
 	WebSocketRouter->AddPreDispatch([this](const struct FRemoteControlWebSocketMessage& Message)
 	{
 		const TArray<FString>* InPassphrase = Message.Header.Find(WebRemoteControlInternalUtils::PassphraseHeader);
+		const FString Passphrase = InPassphrase ? InPassphrase->Last() : FString("");
 		
-		const bool bCanBeDispatched = InPassphrase ? CheckPassphrase(InPassphrase->Last()) : CheckPassphrase("");
+		const bool bCanBeDispatched = CheckPassphrase(Passphrase);
 
 		if (!bCanBeDispatched)
 		{
 			TArray<uint8> Response;
 			FRCRequestWrapper Wrapper;
 			Wrapper.RequestId = Message.MessageId;
+			Wrapper.Passphrase = Passphrase;
+			Wrapper.Verb = "401";
 			
 			WebRemoteControlInternalUtils::CreateUTF8ErrorMessage(FString::Printf(TEXT("Given Passphrase is not correct!")), Wrapper.TCHARBody);
 			WebRemoteControlUtils::SerializeMessage(Wrapper, Response);
@@ -1901,6 +1904,19 @@ void FWebRemoteControlModule::OnSettingsModified(UObject* Settings, FPropertyCha
 		StopWebSocketServer();
 		StartWebSocketServer();
 	}
+
+	/** Letting the Server know what the current state of the Passphrase Usage is. */
+	if (bIsWebSocketServerStarted && bIsWebServerStarted)
+	{
+		const bool bIsOpen = RCSettings->bUseRemoteControlPassphrase;
+		
+		TArray<uint8> Response;
+		const FCheckPassphraseResponse BoolResponse = FCheckPassphraseResponse(!bIsOpen);
+		
+		WebRemoteControlUtils::SerializeMessage(BoolResponse, Response);
+
+		WebSocketServer.Broadcast(Response);
+	}
 }
 
 #endif
diff --git a/Engine/Plugins/VirtualProduction/RemoteControlWebInterface/WebApp/Client/src/reducers/api.ts b/Engine/Plugins/VirtualProduction/RemoteControlWebInterface/WebApp/Client/src/reducers/api.ts
--- a/Engine/Plugins/VirtualProduction/RemoteControlWebInterface/WebApp/Client/src/reducers/api.ts
+++ b/Engine/Plugins/VirtualProduction/RemoteControlWebInterface/WebApp/Client/src/reducers/api.ts
@@ -71,11 +71,15 @@ function _initialize(dispatch: Dispatch, getState: () => { api: ApiState }) {
     })
     .on('opened', (isOpen: boolean) => {
       dispatch(API.STATUS({ isOpen, loading: false }));
-    })
-    .on('passphrase', (keyCorrect: boolean) => {
-      dispatch(API.STATUS({ keyCorrect }));
-
 
+      _api.presets.get();
+      _api.payload.all();
+    })
+    .on('passphrase', (wrongPassphrase: string) => {
+      const isLoginStillCorrect = _passphrase !== wrongPassphrase && _passphrase !== undefined;
+      dispatch(API.STATUS({ keyCorrect: isLoginStillCorrect }));
+      if (!isLoginStillCorrect)
+        _passphrase = undefined;
     })
     .on('loading', (loading: boolean) => {
       dispatch(API.STATUS({ loading }));
@@ -100,6 +104,9 @@ async function _request(method: string, url: string, body: string | object | und
   if (answer.length > 0)
     answer = JSON.parse(answer);
   
+  if (res.status === 401)
+    _dispatch(API.STATUS(answer));
+
   if (!res.ok)
     throw answer;
 
diff --git a/Engine/Plugins/VirtualProduction/RemoteControlWebInterface/WebApp/Server/src/Api.ts b/Engine/Plugins/VirtualProduction/RemoteControlWebInterface/WebApp/Server/src/Api.ts
--- a/Engine/Plugins/VirtualProduction/RemoteControlWebInterface/WebApp/Server/src/Api.ts
+++ b/Engine/Plugins/VirtualProduction/RemoteControlWebInterface/WebApp/Server/src/Api.ts
@@ -26,7 +26,12 @@ export class Api {
 
   @Get('passphrase')
   private async passphrase(req: Request, res: Response) {
-    this.send(res, Promise.resolve({ keyCorrect: await UnrealEngine.checkPassphrase(req.headers['passphrase'].toString()) }) );
+	const passphrase = req.headers['passphrase']?.toString() ?? '';  
+	const keyCorrect = await UnrealEngine.checkPassphrase(passphrase);
+	if (!keyCorrect)
+	  res.status(401);	
+	
+    this.send(res, Promise.resolve({ keyCorrect: keyCorrect}) );
   }
 
   @Get('payloads')
diff --git a/Engine/Plugins/VirtualProduction/RemoteControlWebInterface/WebApp/Server/src/Notify.ts b/Engine/Plugins/VirtualProduction/RemoteControlWebInterface/WebApp/Server/src/Notify.ts
--- a/Engine/Plugins/VirtualProduction/RemoteControlWebInterface/WebApp/Server/src/Notify.ts
+++ b/Engine/Plugins/VirtualProduction/RemoteControlWebInterface/WebApp/Server/src/Notify.ts
@@ -29,7 +29,7 @@ export namespace Notify {
     });
   }
 
-  export function emit(what: 'presets' | 'payloads' | 'connected' | 'loading', ...args: any[]) {
+  export function emit(what: 'presets' | 'payloads' | 'connected' | 'loading' | 'opened', ...args: any[]) {
     io.emit(what, ...args);
   }
 
@@ -46,4 +46,8 @@ export namespace Notify {
   export function emitValuesChanges(preset: string, changes: { [key: string]: PropertyValue }) {
     io.emit('values', preset, changes);
   }
+
+  export function emitPassphraseChanged(wrongPassphrase: string) {
+    io.emit('passphrase', wrongPassphrase);
+  }
 }
\ No newline at end of file
diff --git a/Engine/Plugins/VirtualProduction/RemoteControlWebInterface/WebApp/Server/src/UnrealEngine.ts b/Engine/Plugins/VirtualProduction/RemoteControlWebInterface/WebApp/Server/src/UnrealEngine.ts
--- a/Engine/Plugins/VirtualProduction/RemoteControlWebInterface/WebApp/Server/src/UnrealEngine.ts
+++ b/Engine/Plugins/VirtualProduction/RemoteControlWebInterface/WebApp/Server/src/UnrealEngine.ts
@@ -87,7 +87,7 @@ export namespace UnrealEngine {
   let registered: { [id: string]: boolean };
   let payloads: IPayloads = {};
   let views: { [preset: string]: IView } = {};
-  let workingPassphrase: string = undefined;
+  let workingPassphrases: string[] = [];
 
   export async function initialize() {
     connect();
@@ -137,13 +137,17 @@ export namespace UnrealEngine {
       return false;
     }
 
-    if (workingPassphrase == undefined) {
-      workingPassphrase = passphrase;
-      await pullPresets(true);
-      pullTimer();   
+    if (!workingPassphrases.includes(passphrase)) {
+      workingPassphrases.push(passphrase);
+      
+      if (workingPassphrases.length === 1 || isOpen) {
+        await pullPresets(true);
+        pullTimer();   
+      }
     }
     else {
-      workingPassphrase = passphrase;
+      const index = workingPassphrases.findIndex((value:string):boolean => value === passphrase);
+      workingPassphrases[index] = passphrase;
     }
 
     return true;
@@ -165,7 +169,8 @@ export namespace UnrealEngine {
     Notify.emit('connected', true);
     
     isOpen = await checkPassphrase("");
-    if (isOpen || workingPassphrase !== undefined) {
+    Notify.emit('opened', isOpen);
+    if (isOpen || workingPassphrases.length > 0) {
       await pullPresets(true);
       pullTimer();
     }
@@ -173,7 +178,7 @@ export namespace UnrealEngine {
 
   async function refresh() { 
     try {
-      await pullPresets(true);
+      await pullPresets();
     } catch (error) {
     }
   }
@@ -213,19 +218,36 @@ export namespace UnrealEngine {
 
     try {
       const message = JSON.parse(json);
+      
       if (message.RequestId) {
         const promise = pendings[message.RequestId];
         if (!promise)
-          return;
-
+        return;
+        
         if (Program.logger)
-          LogServer.logLoopback({ RequestId: message.RequestId, Stage: 'WebApp Done' }, true);
+        LogServer.logLoopback({ RequestId: message.RequestId, Stage: 'WebApp Done' }, true);
+        
+        if (message.Verb === "401" && workingPassphrases.includes(message.Passphrase)) {
+          Notify.emitPassphraseChanged(message.Passphrase);
+          workingPassphrases = workingPassphrases.filter(value => value !== message.Passphrase);
+
+          delete pendings[message.RequestId];
+          promise?.(message.ResponseCode);
+          return;
+        }
 
         delete pendings[message.RequestId];
         promise?.(message.ResponseBody);
         return;
       }
 
+      if (isOpen !== message.keyCorrect) {
+        isOpen = !isOpen;
+        Notify.emit('opened', isOpen);
+        await pullPresets(true);
+        pullTimer();
+      }
+
       switch (message.Type) {
         case UnrealApi.PresetEvent.FieldsChanged: {
           const preset = presets[message.PresetId];
@@ -337,7 +359,7 @@ export namespace UnrealEngine {
   function send(message: string, parameters: any, passphrase?: string) {
     verifyConnection();
     const Id = parameters?.RequestId ?? request++;
-    passphrase = passphrase ?? workingPassphrase;
+    passphrase = passphrase ?? workingPassphrases[0];
     connection.send(JSON.stringify({ MessageName: message, Id, Passphrase: passphrase, Parameters: parameters}));
   }
 
diff --git a/Engine/Plugins/VirtualProduction/RemoteControlWebInterface/WebApp/_package.json b/Engine/Plugins/VirtualProduction/RemoteControlWebInterface/WebApp/_package.json
--- a/Engine/Plugins/VirtualProduction/RemoteControlWebInterface/WebApp/_package.json
+++ b/Engine/Plugins/VirtualProduction/RemoteControlWebInterface/WebApp/_package.json
@@ -1,6 +1,6 @@
 {
   "name": "conductor",
-  "version": "0.3.63",
+  "version": "0.3.64",
   "scripts": {
     "postinstall": "concurrently \"cd ./Client && npm install && cd ..\" \"cd ./Server && npm install && cd ..\"",
     "build": "node ./Scripts/build.js",
