RemoteControlWebInterface: Merge from Marmot427
- Commit 1 of 2

diff --git a/Engine/Plugins/VirtualProduction/RemoteControlWebInterface/Source/RemoteControlWebInterface/Private/RCWebInterfaceLibrary.cpp b/Engine/Plugins/VirtualProduction/RemoteControlWebInterface/Source/RemoteControlWebInterface/Private/RCWebInterfaceLibrary.cpp
--- a/Engine/Plugins/VirtualProduction/RemoteControlWebInterface/Source/RemoteControlWebInterface/Private/RCWebInterfaceLibrary.cpp
+++ b/Engine/Plugins/VirtualProduction/RemoteControlWebInterface/Source/RemoteControlWebInterface/Private/RCWebInterfaceLibrary.cpp
@@ -8,6 +8,10 @@
 #include "RCWebInterfacePrivate.h"
 #include "RemoteControlBinding.h"
 #include "RemoteControlPreset.h"
+#include "StructSerializer.h"
+#include "IStructSerializerBackend.h"
+#include "Backends/JsonStructSerializerBackend.h"
+
 
 #define LOCTEXT_NAMESPACE "RemoteControlWebInterface"
 
@@ -20,12 +24,12 @@ TMap<FString, AActor*> URCWebInterfaceBlueprintLibrary::FindMatchingActorsToRebi
 		for (const FString& PropertyId : PropertyIds)
 		{
 			const FGuid PropertyGuid(PropertyId);
-			TWeakPtr<FRemoteControlProperty> ExposedEntity = RCPreset->GetExposedEntity<FRemoteControlProperty>(PropertyGuid);
+			TWeakPtr<FRemoteControlField> ExposedEntity = RCPreset->GetExposedEntity<FRemoteControlField>(PropertyGuid);
 
 			TSet<AActor*> MatchingActorsForProperty;
-			if (TSharedPtr<FRemoteControlProperty> RemoteControlProperty = ExposedEntity.Pin())
+			if (TSharedPtr<FRemoteControlField> RemoteControlField = ExposedEntity.Pin())
 			{
-				TArray<UObject*> BoundObjects = RemoteControlProperty->GetBoundObjects();
+				TArray<UObject*> BoundObjects = RemoteControlField->GetBoundObjects();
 				if (BoundObjects.Num() >= 1)
 				{
 					UObject* OwnerObject = BoundObjects[0];
@@ -50,7 +54,7 @@ TMap<FString, AActor*> URCWebInterfaceBlueprintLibrary::FindMatchingActorsToRebi
 						}
 					}
 				}
-				else if (UClass* SupportedBindingClass = RemoteControlProperty->GetSupportedBindingClass())
+				else if (UClass* SupportedBindingClass = RemoteControlField->GetSupportedBindingClass())
 				{
 					constexpr bool bAllowPIE = false;
 					UWorld* World = RCPreset->GetWorld(bAllowPIE);
@@ -74,7 +78,7 @@ TMap<FString, AActor*> URCWebInterfaceBlueprintLibrary::FindMatchingActorsToRebi
 					}
 				}
 			}
-		
+
 			if (MatchesIntersection.Num() == 0)
 			{
 				MatchesIntersection = MatchingActorsForProperty;
@@ -110,10 +114,10 @@ FString URCWebInterfaceBlueprintLibrary::GetOwnerActorLabel(const FString& Prese
 		for (const FString& PropertyId : PropertyIds)
 		{
 			const FGuid PropertyGuid(PropertyId);
-			TWeakPtr<FRemoteControlProperty> ExposedEntity = RCPreset->GetExposedEntity<FRemoteControlProperty>(PropertyGuid);
-			if (TSharedPtr<FRemoteControlProperty> RemoteControlProperty = ExposedEntity.Pin())
+			TWeakPtr<FRemoteControlField> ExposedEntity = RCPreset->GetExposedEntity<FRemoteControlField>(PropertyGuid);
+			if (TSharedPtr<FRemoteControlField> RemoteControlField = ExposedEntity.Pin())
 			{
-				TArray<UObject*> BoundObjects = RemoteControlProperty->GetBoundObjects();
+				TArray<UObject*> BoundObjects = RemoteControlField->GetBoundObjects();
 				if (BoundObjects.Num() >= 1)
 				{
 					UObject* OwnerObject = BoundObjects[0];
@@ -151,10 +155,10 @@ void URCWebInterfaceBlueprintLibrary::RebindProperties(const FString& PresetId,
 		for (const FString& PropertyId : PropertyIds)
 		{
 			const FGuid PropertyGuid(PropertyId);
-			TWeakPtr<FRemoteControlProperty> ExposedEntity = RCPreset->GetExposedEntity<FRemoteControlProperty>(PropertyGuid);
-			if (TSharedPtr<FRemoteControlProperty> RemoteControlProperty = ExposedEntity.Pin())
+			TWeakPtr<FRemoteControlField> ExposedEntity = RCPreset->GetExposedEntity<FRemoteControlField>(PropertyGuid);
+			if (TSharedPtr<FRemoteControlField> RemoteControlField = ExposedEntity.Pin())
 			{
-				RemoteControlProperty->BindObject(NewOwner);
+				RemoteControlField->BindObject(NewOwner);
 			}
 		}
 	}
@@ -183,4 +187,66 @@ FString URCWebInterfaceBlueprintLibrary::GetActorNameOrLabel(const AActor* Actor
 #endif
 }
 
+TMap<AActor*, FString> URCWebInterfaceBlueprintLibrary::FindAllActorsOfClass(UClass* Class)
+{
+	UWorld* World = URemoteControlPreset::GetWorld(nullptr, false);
+	if (!World || !Class)
+	{
+		return {};
+	}
+
+	TArray<AActor*> FoundActors;
+	UGameplayStatics::GetAllActorsOfClass(World, Class, FoundActors);
+
+	TMap<AActor*, FString> Results;
+	for (AActor* Actor : FoundActors)
+	{
+		Results.Add(Actor, GetActorNameOrLabel(Actor));
+	}
+
+	return Results;
+}
+
+AActor* URCWebInterfaceBlueprintLibrary::SpawnActor(UClass* Class)
+{
+	UWorld* World = URemoteControlPreset::GetWorld(nullptr, false);
+	if (!World || !Class)
+	{
+		return nullptr;
+	}
+
+	return World->SpawnActor<AActor>(Class, FActorSpawnParameters());
+}
+
+TMap<AActor*, FString> URCWebInterfaceBlueprintLibrary::GetValuesOfActorsByClass(UClass* Class)
+{
+	UWorld* World = URemoteControlPreset::GetWorld(nullptr, false);
+	if (!World || !Class)
+	{
+		return {};
+	}
+
+	TArray<AActor*> FoundActors;
+	UGameplayStatics::GetAllActorsOfClass(World, Class, FoundActors);
+
+	TMap<AActor*, FString> Values;
+	Values.Reserve(FoundActors.Num());
+
+	for (AActor* Actor : FoundActors)
+	{
+		TArray<uint8> WorkingBuffer;
+		FMemoryWriter Writer(WorkingBuffer);
+		FJsonStructSerializerBackend SerializerBackend(Writer, EStructSerializerBackendFlags::Default);
+
+		FStructSerializer::Serialize(Actor, *Class, SerializerBackend, FStructSerializerPolicies());
+
+		// Add 16bit \0
+		WorkingBuffer.AddDefaulted(2);
+		Values.Add(Actor, TCHAR_TO_UTF8(WorkingBuffer.GetData()));
+	}
+
+	return Values;
+}
+
+
 #undef LOCTEXT_NAMESPACE
\ No newline at end of file
diff --git a/Engine/Plugins/VirtualProduction/RemoteControlWebInterface/Source/RemoteControlWebInterface/Private/RCWebInterfaceLibrary.h b/Engine/Plugins/VirtualProduction/RemoteControlWebInterface/Source/RemoteControlWebInterface/Private/RCWebInterfaceLibrary.h
--- a/Engine/Plugins/VirtualProduction/RemoteControlWebInterface/Source/RemoteControlWebInterface/Private/RCWebInterfaceLibrary.h
+++ b/Engine/Plugins/VirtualProduction/RemoteControlWebInterface/Source/RemoteControlWebInterface/Private/RCWebInterfaceLibrary.h
@@ -31,6 +31,24 @@ public:
 	UFUNCTION(BlueprintCallable, Category = "RemoteControlWebInterface")
 	static void RebindProperties(const FString& PresetId, const TArray<FString>& PropertyIds, AActor* NewOwner);
 
+	/**
+	 * Shortcut function to find all actors of a class
+	 */
+	UFUNCTION(BlueprintCallable, Category = "RemoteControlWebInterface")
+	static TMap<AActor*, FString> FindAllActorsOfClass(UClass* Class);
+
+	/**
+	 * Shortcut function to spawn an actor of a class
+	 */
+	UFUNCTION(BlueprintCallable, Category = "RemoteControlWebInterface")
+	static AActor* SpawnActor(UClass* Class);
+
+	/**
+	 * Gets all properties values (as a json) of all actors of type Class
+	 */
+	UFUNCTION(BlueprintCallable, Category = "RemoteControlWebInterface")
+	static TMap<AActor*, FString> GetValuesOfActorsByClass(UClass* Class);
+
 private:
 	static class URemoteControlPreset* GetPreset(const FString& PresetId);
 
diff --git a/Engine/Plugins/VirtualProduction/RemoteControlWebInterface/Source/RemoteControlWebInterface/RemoteControlWebInterface.Build.cs b/Engine/Plugins/VirtualProduction/RemoteControlWebInterface/Source/RemoteControlWebInterface/RemoteControlWebInterface.Build.cs
--- a/Engine/Plugins/VirtualProduction/RemoteControlWebInterface/Source/RemoteControlWebInterface/RemoteControlWebInterface.Build.cs
+++ b/Engine/Plugins/VirtualProduction/RemoteControlWebInterface/Source/RemoteControlWebInterface/RemoteControlWebInterface.Build.cs
@@ -20,7 +20,7 @@ public class RemoteControlWebInterface : ModuleRules
 				"RemoteControl",
 				"Sockets",
 				"WebRemoteControl",
-				"WebSocketNetworking"
+				"WebSocketNetworking",
 			}
 		);
 
@@ -28,6 +28,7 @@ public class RemoteControlWebInterface : ModuleRules
 			new string[] 
 			{
 				"RemoteControlCommon",
+				"Serialization",
 			}
 		);
 
diff --git a/Engine/Plugins/VirtualProduction/RemoteControlWebInterface/WebApp/Client/Public/index.html b/Engine/Plugins/VirtualProduction/RemoteControlWebInterface/WebApp/Client/Public/index.html
--- a/Engine/Plugins/VirtualProduction/RemoteControlWebInterface/WebApp/Client/Public/index.html
+++ b/Engine/Plugins/VirtualProduction/RemoteControlWebInterface/WebApp/Client/Public/index.html
@@ -5,7 +5,7 @@
   <link rel="icon" type="image/png" sizes="32x32" href="/images/favicon/32x32.png">
   <link rel="icon" type="image/png" sizes="16x16" href="/images/favicon/16x16.png">
   <link rel="manifest" href="/manifest.json">
-  <meta name="viewport" content="user-scalable=no" />
+  <meta name="viewport" content="width=device-width,initial-scale=1,user-scalable=no" />
 </head>
 <body>
   <div id="root"></div>
diff --git a/Engine/Plugins/VirtualProduction/RemoteControlWebInterface/WebApp/Client/_package.json b/Engine/Plugins/VirtualProduction/RemoteControlWebInterface/WebApp/Client/_package.json
--- a/Engine/Plugins/VirtualProduction/RemoteControlWebInterface/WebApp/Client/_package.json
+++ b/Engine/Plugins/VirtualProduction/RemoteControlWebInterface/WebApp/Client/_package.json
@@ -6,6 +6,7 @@
   },
   "dependencies": {
     "@fortawesome/fontawesome-svg-core": "^1.2.30",
+    "@fortawesome/free-regular-svg-icons": "^6.0.0",
     "@fortawesome/free-solid-svg-icons": "^5.14.0",
     "@fortawesome/react-fontawesome": "^0.1.11",
     "dot-prop-immutable": "^2.1.0",
diff --git a/Engine/Plugins/VirtualProduction/RemoteControlWebInterface/WebApp/Client/src/App.tsx b/Engine/Plugins/VirtualProduction/RemoteControlWebInterface/WebApp/Client/src/App.tsx
--- a/Engine/Plugins/VirtualProduction/RemoteControlWebInterface/WebApp/Client/src/App.tsx
+++ b/Engine/Plugins/VirtualProduction/RemoteControlWebInterface/WebApp/Client/src/App.tsx
@@ -1,8 +1,8 @@
 import React from 'react';
 import { connect } from 'react-redux';
 import { _api, ReduxState } from './reducers';
-import { ConnectionStatus, Tabs, Screen, Stack, RenameModal, AlertModal, IconModal, PropertiesDrawer, Login } from './components';
-import { IPreset, IView, ITab, TabLayout, WidgetTypes, ICustomStackProperty, IPanelType, ScreenType, IPanel, ICustomStackListItem } from './shared';
+import { ConnectionStatus, Tabs, Screen, Stack, RenameModal, AlertModal, IconModal, PropertiesDrawer, Login, SignalIcon, EmptyTab } from './components';
+import { IPreset, IView, ITab, TabLayout, WidgetTypes, ICustomStackProperty, IPanelType, ScreenType, IPanel, ICustomStackListItem, IHistory, ConnectionSignal } from './shared';
 import { fas, IconName } from '@fortawesome/free-solid-svg-icons';
 import { BeforeCapture, DragDropContext, DropResult } from 'react-beautiful-dnd';
 import { FontAwesomeIcon } from '@fortawesome/react-fontawesome';
@@ -19,7 +19,10 @@ type PropsFromState = Partial<{
   view: IView;
   presets: IPreset[]
   preset: IPreset;
-
+  signal?: ConnectionSignal;
+  lockedUI: boolean,
+  undo: IHistory[],
+  redo: IHistory[],
   initialize?: () => void;
 }>;
 
@@ -31,6 +34,10 @@ const mapStateToProps = (state: ReduxState): PropsFromState => ({
   presets: _.values(state.api.presets),
   preset: state.api.presets[state.api.preset],
   view: state.api.view,
+  lockedUI: state.api.lockedUI,
+  signal: state.api.status.signal,
+  undo: state.api.undo,
+  redo: state.api.redo,
 });
 
 const mapDispatchToProps = (dispatch): PropsFromState => ({
@@ -269,8 +276,8 @@ export class App extends React.Component<PropsFromState, State> {
 
   onNewTab = (panels: IPanel[] = []) => {
     const { view } = this.props;
-    
-    this.createNewTabInternal(panels);
+
+    this.createNewTabInternal(panels, null, TabLayout.Empty);
     _api.views.set(view);
 
     this.onTabChange(view.tabs.length - 1);
@@ -289,13 +296,13 @@ export class App extends React.Component<PropsFromState, State> {
     this.setState({ editable: true });
   }
 
-  createNewTabInternal = (panels: IPanel[] = [], name?: string) => {
+  createNewTabInternal = (panels: IPanel[] = [], name?: string, layout = TabLayout.Stack) => {
     const { view } = this.props;
 
     view.tabs.push({
       name: name ?? this.getNewTabName(),
       icon: this.getRandomIcon(),
-      layout: TabLayout.Stack,
+      layout,
       panels
     });
   }
@@ -315,32 +322,40 @@ export class App extends React.Component<PropsFromState, State> {
     this.setState({ editable: true });
   }
 
-  onAddSnapshotTab = () => {
+  onScreenAdd = (type: ScreenType) => {
     const { view } = this.props;
+    let tab = view.tabs[this.state.tab];
 
-    view.tabs.push({
-      name: 'Snapshot',
-      icon: 'save',
-      layout: TabLayout.Screen,
-      screen: { type: ScreenType.Snapshot },
-    });
+    delete tab.panels;
 
-    _api.views.set(view);
-    this.onTabChange(view.tabs.length - 1);
-  }
+    switch (type) {
+      case ScreenType.ColorCorrection:
+        tab = { ...tab, name: 'Color Correction', icon: 'palette', layout: TabLayout.Screen, screen: { type: ScreenType.ColorCorrection }, };
+        break;
 
-  onAddSequencerTab = () => {
-    const { view } = this.props;
+      case ScreenType.LightCards:
+        tab = { ...tab, name: 'Light Cards', icon: 'adjust', layout: TabLayout.Screen, screen: { type: ScreenType.LightCards }, };
+        break;
 
-    view.tabs.push({
-      name: 'Sequences',
-      icon: 'play',
-      layout: TabLayout.Screen,
-      screen: { type: ScreenType.Sequencer },
-    });
+      case ScreenType.Playlist:
+        tab = { ...tab, name: 'Playlists', icon: 'play', layout: TabLayout.Screen, screen: { type: ScreenType.Playlist }, };
+        break;
+
+      case ScreenType.Sequencer:
+        tab = { ...tab, name: 'Sequences', icon: 'play', layout: TabLayout.Screen, screen: { type: ScreenType.Sequencer }, };
+        break;
+
+      case ScreenType.Snapshot:
+        tab = { ...tab, name: 'Snapshot', icon: 'save', layout: TabLayout.Screen, screen: { type: ScreenType.Snapshot }, };
+        break;
+
+      case ScreenType.Stack:
+        tab = { ...tab, panels: [], layout: TabLayout.Stack };
+        break;
+    }
 
+    view.tabs[this.state.tab] = tab;
     _api.views.set(view);
-    this.onTabChange(view.tabs.length - 1);
   }
 
   onSetTabsDrawer = (editable: boolean) => {
@@ -521,6 +536,8 @@ export class App extends React.Component<PropsFromState, State> {
 
     const index = this.state.tab;
     const tab = view.tabs[index];
+    if (!tab)
+      return null;
 
     let visible = true;
     const style: React.CSSProperties = {};
@@ -530,7 +547,7 @@ export class App extends React.Component<PropsFromState, State> {
     }
 
     let content = null;
-    switch (tab?.layout) {
+    switch (tab.layout) {
       case TabLayout.Stack:
         content = this.renderStack(tab, visible, index);
         break;
@@ -538,6 +555,9 @@ export class App extends React.Component<PropsFromState, State> {
       case TabLayout.Screen:
         content = this.renderScreen(tab, visible);
         break;
+
+      case TabLayout.Empty:
+        content = <EmptyTab onAddScreen={this.onScreenAdd} />;
     }
 
     return (
@@ -586,6 +606,11 @@ export class App extends React.Component<PropsFromState, State> {
     _api.views.set(view);
   }
 
+  onLockedChange = () => {
+    const { lockedUI } = this.props;
+    _api.lockUI(!lockedUI);
+  }
+
   checkPassphrase = async (passphrase: string): Promise<boolean> => {
     return await _api.passphrase.login(passphrase);
   }
@@ -598,6 +623,7 @@ export class App extends React.Component<PropsFromState, State> {
 
     if (!keyCorrect && !isOpen)
       return <Login />;
+
     if (loading) {
       return (
         <div className="fullscreen">
@@ -616,24 +642,35 @@ export class App extends React.Component<PropsFromState, State> {
   }
 
   render() {
-    const { view, preset, presets } = this.props;
+    const { view, preset, presets, lockedUI, undo, redo, signal } = this.props;
     const { tab, editable, selected } = this.state;
 
+    let appClassName = 'tabs-closed ';
+
+    if (editable)
+      appClassName = appClassName.replace('closed', 'open');
+
+    if (lockedUI)
+      appClassName += 'locked ';  
+
     return (
       <DragDropContext onBeforeCapture={this.onBeforeCapture}
                        onDragEnd={this.onDragEnd}>
         <div id="app"
-             className={`${editable ? 'tabs-open' : ''} `}
+             className={appClassName}
              tabIndex={1}
              ref={this.appRef}
              onKeyDown={this.onKeyDown} >
-
           {this.renderStatus()}
 
+          {lockedUI && 
+            <div className='locked-layer' />
+          }
           <nav id="top-bar">
             <div className="app-icon">
               <img src="/images/favicon/32x32.png" alt="Unreal Engine" />
             </div>
+           
             <div className="tab edit-tab mode-toggle">
               <label className="switch toggle-mode">
                 <input type="checkbox" checked={!!editable} onChange={() => this.onSetTabsDrawer(!editable)} />
@@ -647,7 +684,11 @@ export class App extends React.Component<PropsFromState, State> {
                   onChange={this.onTabChange}
                   onNewTab={this.onNewTab}
                   editable={editable}
-                  onDeleteTab={this.onTabDelete} />
+                  locked={lockedUI}
+                  undo={!!undo.length}
+                  redo={!!redo.length}
+                  onDeleteTab={this.onTabDelete}
+                  onLockedChange={this.onLockedChange} />
           </nav>
 
           {!!preset &&
@@ -660,9 +701,7 @@ export class App extends React.Component<PropsFromState, State> {
                               selected={selected}
                               onChangeIcon={this.onTabIconChange}
                               onRenameTabModal={this.onTabRename}
-                              onDuplicateTab={this.onDuplicateTab}
-                              onAddSnapshotTab={this.onAddSnapshotTab}
-                              onAddSequencerTab={this.onAddSequencerTab}
+                              onDuplicateTab={this.onDuplicateTab}                              
                               onPresetChange={this.onPresetChanged}
                               onUpdateView={this.onUpdateView}
                               onSelected={this.onSelected}
@@ -670,6 +709,7 @@ export class App extends React.Component<PropsFromState, State> {
                               lockWidget={widget => this.widget = widget} />
           }
           {this.renderTab()}
+          <SignalIcon signal={signal} />
         </div>
       </DragDropContext>
     );
diff --git a/Engine/Plugins/VirtualProduction/RemoteControlWebInterface/WebApp/Client/src/IconLibrary.ts b/Engine/Plugins/VirtualProduction/RemoteControlWebInterface/WebApp/Client/src/IconLibrary.ts
--- a/Engine/Plugins/VirtualProduction/RemoteControlWebInterface/WebApp/Client/src/IconLibrary.ts
+++ b/Engine/Plugins/VirtualProduction/RemoteControlWebInterface/WebApp/Client/src/IconLibrary.ts
@@ -1,4 +1,5 @@
-import { library } from '@fortawesome/fontawesome-svg-core';
+import { IconPack, library } from '@fortawesome/fontawesome-svg-core';
 import { fas } from '@fortawesome/free-solid-svg-icons';
+import { far } from '@fortawesome/free-regular-svg-icons';
 
-library.add(fas);
\ No newline at end of file
+library.add(fas as IconPack, far as IconPack);
\ No newline at end of file
diff --git a/Engine/Plugins/VirtualProduction/RemoteControlWebInterface/WebApp/Client/src/components/EmptyTab.tsx b/Engine/Plugins/VirtualProduction/RemoteControlWebInterface/WebApp/Client/src/components/EmptyTab.tsx
new file mode 100644
--- /dev/null
+++ b/Engine/Plugins/VirtualProduction/RemoteControlWebInterface/WebApp/Client/src/components/EmptyTab.tsx
@@ -0,0 +1,34 @@
+import React from 'react';
+import { ScreenType } from 'src/shared';
+
+
+type Props = {
+  onAddScreen: (screen: ScreenType) => void;
+}
+
+export class EmptyTab extends React.Component<Props> {
+
+  screens = [
+    { type: ScreenType.Stack, title: 'Build Your Own UI' },
+    { type: ScreenType.ColorCorrection, title: 'Color Correction' },
+    { type: ScreenType.Playlist, title: 'Playlist' },
+    { type: ScreenType.Snapshot, title: 'Snapshot' },
+    { type: ScreenType.Sequencer, title: 'Sequencer' },
+  ];
+
+  render() {
+    return (
+      <div className="panel">
+        <div className="empty-tab-container">
+          <ul className="screens-list">
+            {this.screens.map(({ type, title }) => (
+              <li className="screen-type-item" key={type}>
+                <button className="btn btn-secondary" onClick={() => this.props.onAddScreen(type)}>{title}</button>
+              </li>
+            ))}
+          </ul>
+        </div>
+      </div>
+    );
+  }
+};
\ No newline at end of file
diff --git a/Engine/Plugins/VirtualProduction/RemoteControlWebInterface/WebApp/Client/src/components/Login.tsx b/Engine/Plugins/VirtualProduction/RemoteControlWebInterface/WebApp/Client/src/components/Login.tsx
--- a/Engine/Plugins/VirtualProduction/RemoteControlWebInterface/WebApp/Client/src/components/Login.tsx
+++ b/Engine/Plugins/VirtualProduction/RemoteControlWebInterface/WebApp/Client/src/components/Login.tsx
@@ -1,39 +1,8 @@
 import React from 'react';
-import crypto from 'crypto';
 import { ReactComponent as Logo } from '../assets/ue_logo.svg';
-import { _api } from '../reducers';
 
 
-type Props = {
-
-}
-
-type State = {
-  passphrase: string;
-  error: boolean;
-}
-
-export class Login extends React.Component<Props, State> {
-  state: State = {
-    passphrase: '',
-    error: false,
-  }
-
-  onChange = (passphrase: string) => {
-    this.setState({ passphrase });
-  }
-
-  onKeyPress = async (event: React.KeyboardEvent<HTMLInputElement>) => {
-    if (event.key !== "Enter")
-      return;
-
-    const { passphrase } = this.state;
-    const secured = crypto.createHash('md5')
-                          .update(passphrase).digest('hex');
-
-    const error = await _api.passphrase.login(secured);
-    this.setState({ error });
-  }
+export class Login extends React.Component {
 
   render() {
     return (
@@ -46,13 +15,9 @@ export class Login extends React.Component<Props, State> {
         </div>
         <div className='form'>
           Password
-          <input onChange={e => this.onChange(e.target.value)}
-                  onKeyPress={this.onKeyPress}
-                  type='password' />
-
-          {this.state.error &&
-            <label className='login-status'>Incorrect Password</label>
-          }
+          <input type='password'/>
+          <div className='hint'>Configure the password in: <span className="breadcrumb">Project Settings / Authentication / Password</span></div>
+          <label className='login-status'>Incorrect Password</label>
         </div>
      </div>
     );
diff --git a/Engine/Plugins/VirtualProduction/RemoteControlWebInterface/WebApp/Client/src/components/Properties/Pressets.tsx b/Engine/Plugins/VirtualProduction/RemoteControlWebInterface/WebApp/Client/src/components/Properties/Pressets.tsx
--- a/Engine/Plugins/VirtualProduction/RemoteControlWebInterface/WebApp/Client/src/components/Properties/Pressets.tsx
+++ b/Engine/Plugins/VirtualProduction/RemoteControlWebInterface/WebApp/Client/src/components/Properties/Pressets.tsx
@@ -2,6 +2,9 @@ import { FontAwesomeIcon } from '@fortawesome/react-fontawesome';
 import React from 'react';
 import { IPreset } from 'src/shared';
 import { Search } from '../controls';
+import { _api } from '../../reducers';
+import _ from 'lodash';
+import { IconProp } from '@fortawesome/fontawesome-svg-core';
 
 
 type Props = {
@@ -12,10 +15,39 @@ type Props = {
   onSearch: (value: string) => void;
 }
 
-export class Pressets extends React.Component<Props> {
+type State = {
+  spin: string;
+}
+
+export class Pressets extends React.Component<Props, State> {
+
+  state: State = {
+    spin: null,
+  }
+
+  onFavoriteToggled = (e: React.MouseEvent, preset: IPreset) => {
+    const { spin } = this.state;
+    e.stopPropagation();
+
+    if (spin !== preset.ID)
+      this.setState({ spin: preset.ID });
+
+    this.onDebouncedFavoriteClick(preset);
+  }
+
+  onFavoritePreset = async (preset: IPreset) => {
+    if (!preset.IsFavorite)
+      this.props.onPresetChange(preset);
+
+    await _api.presets.favorite(preset.ID, !preset.IsFavorite);
+    this.setState({ spin: null });
+  }
+
+  onDebouncedFavoriteClick = _.debounce(this.onFavoritePreset, 100);
 
   render() {
     const { preset, pressets } = this.props;
+    const { spin } = this.state;
 
     return (
       <div className="presets-tab">
@@ -23,12 +55,23 @@ export class Pressets extends React.Component<Props> {
         <div className="presets-wrapper">
           {pressets.map(p => {
             let className = 'btn preset-btn ';
+            let spinIconClassName = 'spin-icon ';
+            let icon: IconProp = ['far', 'star'];
+
+            if (p.IsFavorite)
+              icon = ['fas', 'star'];
 
-            if (p.ID === preset?.ID)
+            if (p.ID === preset?.ID) {
               className += 'active';
+            }
+
+            if (spin === p.ID)
+              spinIconClassName += 'visible';
 
             return (
-              <div key={p?.ID} className={className} onClick={this.props.onPresetChange.bind(this, p)}>
+              <div key={p?.ID} className={className} onClick={() => this.props.onPresetChange(p)}>
+                <FontAwesomeIcon className={spinIconClassName} icon={['fas', 'sync']} spin />
+                <FontAwesomeIcon className='favorite-icon' icon={icon} onClick={e => this.onFavoriteToggled(e, p)} />
                 {p?.Name}
                 <div className="item-icon">
                   <FontAwesomeIcon icon={['fas', 'external-link-square-alt']} />
diff --git a/Engine/Plugins/VirtualProduction/RemoteControlWebInterface/WebApp/Client/src/components/Properties/Tab.tsx b/Engine/Plugins/VirtualProduction/RemoteControlWebInterface/WebApp/Client/src/components/Properties/Tab.tsx
--- a/Engine/Plugins/VirtualProduction/RemoteControlWebInterface/WebApp/Client/src/components/Properties/Tab.tsx
+++ b/Engine/Plugins/VirtualProduction/RemoteControlWebInterface/WebApp/Client/src/components/Properties/Tab.tsx
@@ -2,12 +2,9 @@ import React from 'react';
 
 
 type Props = {
-
   onChangeIcon: () => void;
   onRenameTabModal: () => void;
   onDuplicateTab: () => void;
-  onAddSnapshotTab: () => void;
-  onAddSequencerTab: () => void;
 }
 
 export class Tab extends React.Component<Props> {
@@ -18,8 +15,6 @@ export class Tab extends React.Component<Props> {
           <button tabIndex={-1} className="btn btn-secondary" onClick={this.props.onChangeIcon}>Change tab icon</button>
           <button tabIndex={-1} className="btn btn-secondary" onClick={this.props.onRenameTabModal}>Rename Tab</button>
           <button tabIndex={-1} className="btn btn-secondary" onClick={this.props.onDuplicateTab}>Duplicate Tab</button>
-          <button tabIndex={-1} className="btn btn-secondary" onClick={this.props.onAddSnapshotTab}>Add Snapshot Tab</button>
-          <button tabIndex={-1} className="btn btn-secondary" onClick={this.props.onAddSequencerTab}>Add Sequencer Tab</button>
         </div>
       </div>
     );
diff --git a/Engine/Plugins/VirtualProduction/RemoteControlWebInterface/WebApp/Client/src/components/PropertiesDrawer.tsx b/Engine/Plugins/VirtualProduction/RemoteControlWebInterface/WebApp/Client/src/components/PropertiesDrawer.tsx
--- a/Engine/Plugins/VirtualProduction/RemoteControlWebInterface/WebApp/Client/src/components/PropertiesDrawer.tsx
+++ b/Engine/Plugins/VirtualProduction/RemoteControlWebInterface/WebApp/Client/src/components/PropertiesDrawer.tsx
@@ -1,11 +1,11 @@
 import React, { Component } from 'react';
 import { _api } from 'src/reducers';
 import { IPreset, IGroup, IExposedProperty, WidgetTypes, ICustomStackProperty, IPanel, ICustomStackTabs, IExposedFunction, 
-        PropertyType, IView, ITab, TabLayout, ICustomStackWidget, WidgetType, IColorPickerList, ICustomStackItem, ICustomStackListItem } from 'src/shared';
+        PropertyType, IView, ITab, TabLayout, ICustomStackWidget, WidgetType, IColorPickerList, ICustomStackListItem } from 'src/shared';
 import { Draggable, Droppable, DraggableStateSnapshot } from 'react-beautiful-dnd';
 import { FontAwesomeIcon } from '@fortawesome/react-fontawesome';
 import { AlertModal } from '.';
-import { SlidersWidget, SliderWidget, DropdownWidget, ImageSelectorWidget, ToggleWidget, VectorWidget } from './Widgets';
+import { SlidersWidget, SliderWidget, DropdownWidget, ImageSelectorWidget, ToggleWidget, VectorWidget, AssetWidget } from './Widgets';
 import * as Tabs from './Properties';
 import { ColorPicker, Button, Text, Dial, JoysticksWrapper, DialsWrapper } from './controls';
 import { WidgetUtilities } from 'src/utilities';
@@ -51,8 +51,6 @@ type Props = {
   onDuplicateTab: () => void;
   onUpdateView: (update?: boolean) => void;
   onPresetChange: (preset: IPreset) => void;
-  onAddSnapshotTab: () => void;
-  onAddSequencerTab: () => void;
   onSelected: (selected: string) => void;
   onNewTabs: (items: ICustomStackListItem[]) => void;
 };
@@ -368,6 +366,9 @@ export class PropertiesDrawer extends Component<Props, State> {
       case WidgetTypes.Dials:
         return <DialsWrapper type={widget.propertyType} />;
 
+      case WidgetTypes.Asset:
+        return <AssetWidget />;
+
       default:
         return null;
     }
@@ -478,7 +479,7 @@ export class PropertiesDrawer extends Component<Props, State> {
     let checked = false;
 
     if (exposed?.Metadata?.Alpha === '1')
-      checked = true;    
+      checked = true;
 
     return (
       <>
@@ -678,9 +679,16 @@ export class PropertiesDrawer extends Component<Props, State> {
   }
 
   renderVectorEdit = (available: boolean) => {
+    const { preset } = this.props;
     const widget = this.state.widget as ICustomStackProperty;
     const vectorModes = widget?.widgets || [];
 
+    const exposed = preset?.Exposed[widget?.property] as IExposedProperty;
+    let checked = false;
+
+    if (exposed?.Metadata?.Proportionally === '1')
+      checked = true;
+
     return (
       <>
         {available &&
@@ -694,15 +702,19 @@ export class PropertiesDrawer extends Component<Props, State> {
               <label>Modes</label>
               <div className="btn-group vector-btn-group">
                 {widget.propertyType !== PropertyType.Rotator &&
-                  <button className={`btn ${vectorModes.includes(WidgetTypes.Joystick) ? 'btn-primary' : 'btn-secondary'}`} 
+                  <button className={`btn ${vectorModes.includes(WidgetTypes.Joystick) ? 'btn-primary' : 'btn-secondary'}`}
                           onClick={() => this.onVectorModeChange(widget, WidgetTypes.Joystick)}>Joystick</button>
                 }
-                <button className={`btn ${vectorModes.includes(WidgetTypes.Dial) ? 'btn-primary' : 'btn-secondary'}`} 
+                <button className={`btn ${vectorModes.includes(WidgetTypes.Dial) ? 'btn-primary' : 'btn-secondary'}`}
                         onClick={() => this.onVectorModeChange(widget, WidgetTypes.Dial)}>Dial</button>
-                <button className={`btn ${vectorModes.includes(WidgetTypes.Sliders) ? 'btn-primary' : 'btn-secondary'}`} 
+                <button className={`btn ${vectorModes.includes(WidgetTypes.Sliders) ? 'btn-primary' : 'btn-secondary'}`}
                         onClick={() => this.onVectorModeChange(widget, WidgetTypes.Sliders)}>Sliders</button>
               </div>
             </div>
+            <div className="properties-field">
+              <label>Proportionally</label>
+              <input className="checkbox" type="checkbox" onChange={e => this.onMetadataChange(widget, 'Proportionally', e.target.checked ? '1' : '0')} checked={checked} />
+            </div>
           </div>
         }
         {this.renderDelete()}
@@ -715,10 +727,17 @@ export class PropertiesDrawer extends Component<Props, State> {
   }
 
   renderSlidersEdit = (available: boolean) => {
+    const { preset } = this.props;
     const widget = this.state.widget as ICustomStackProperty;
     const keys = WidgetUtilities.getPropertyKeys(widget?.propertyType);
     const vectorModes = widget?.widgets || [];
 
+    const exposed = preset?.Exposed[widget?.property] as IExposedProperty;
+    let checked = false;
+
+    if (exposed?.Metadata?.Proportionally === '1')
+      checked = true;
+
     return (
       <>
         {available &&
@@ -735,6 +754,10 @@ export class PropertiesDrawer extends Component<Props, State> {
                 )}
               </div>
             </div>
+            <div className="properties-field">
+              <label>Proportionally</label>
+              <input className="checkbox" type="checkbox" onChange={e => this.onMetadataChange(widget, 'Proportionally', e.target.checked ? '1' : '0')} checked={checked} />
+            </div>
           </div>
         }
         {this.renderDelete()}
@@ -1077,9 +1100,7 @@ export class PropertiesDrawer extends Component<Props, State> {
       case TabType.Tab:
         return <Tabs.Tab onChangeIcon={this.props.onChangeIcon}
                          onRenameTabModal={this.props.onRenameTabModal}
-                         onDuplicateTab={this.props.onDuplicateTab}
-                         onAddSnapshotTab={this.props.onAddSnapshotTab}
-                         onAddSequencerTab={this.props.onAddSequencerTab} />;
+                         onDuplicateTab={this.props.onDuplicateTab} />;
 
       case TabType.Widgets:
         return <Tabs.Widgets widgets={filteredWidgets}
diff --git a/Engine/Plugins/VirtualProduction/RemoteControlWebInterface/WebApp/Client/src/components/Screen.tsx b/Engine/Plugins/VirtualProduction/RemoteControlWebInterface/WebApp/Client/src/components/Screen.tsx
--- a/Engine/Plugins/VirtualProduction/RemoteControlWebInterface/WebApp/Client/src/components/Screen.tsx
+++ b/Engine/Plugins/VirtualProduction/RemoteControlWebInterface/WebApp/Client/src/components/Screen.tsx
@@ -38,6 +38,15 @@ export class Screen extends React.Component<Props & PropsFromState> {
 
       case ScreenType.Sequencer:
         return <Screens.Sequencer {...props} />;
+
+      case ScreenType.Playlist:
+        return <Screens.Playlist {...props} />;
+      
+      case ScreenType.ColorCorrection:
+        return <Screens.ColorCorrection {...props} />;
+      
+      case ScreenType.LightCards:
+        return <Screens.LightCards {...props} />;
     }
 
     return null;
diff --git a/Engine/Plugins/VirtualProduction/RemoteControlWebInterface/WebApp/Client/src/components/Screens/ColorCorrection/ColorCorrection.tsx b/Engine/Plugins/VirtualProduction/RemoteControlWebInterface/WebApp/Client/src/components/Screens/ColorCorrection/ColorCorrection.tsx
new file mode 100644
--- /dev/null
+++ b/Engine/Plugins/VirtualProduction/RemoteControlWebInterface/WebApp/Client/src/components/Screens/ColorCorrection/ColorCorrection.tsx
@@ -0,0 +1,906 @@
+import React from 'react';
+import { _api } from 'src/reducers';
+import _ from 'lodash';
+import { Button, ColorPicker, Search, Slider, TabPane, Tabs, VectorControl, TimeControl, SliderWheel, ValueInput } from '../../controls';
+import { AlertModal } from '../../../components';
+import { AssetWidget } from '../../Widgets';
+import { CorrectionColorPicker, PaginationContent } from './';
+import { PropertyType, PropertyValue } from 'src/shared';
+import { FontAwesomeIcon } from '@fortawesome/react-fontawesome';
+import { IconProp } from '@fortawesome/fontawesome-svg-core';
+
+export enum Section {
+  ColorCorrection = 'ColorCorrection',
+  LightCards      = 'LightCards',
+}
+
+export enum CCRMode {
+  Global      = 'Global',
+  Highlights  = 'Highlights',
+  Midtones    = 'Midtones',
+  Shadows     = 'Shadows',
+  Color       = 'Color',
+}
+
+export enum ColorProperty {
+  Gain        = 'Gain',
+  Gamma       = 'Gamma',
+  Saturation  = 'Saturation',
+  Contrast    = 'Contrast',
+  Offset      = 'Offset',
+}
+
+export enum ColorMode {
+  Rgb         = 'RGB',
+  Hsv         = 'HSV'
+}
+
+enum ActorsGroup {
+  All       = 'All',
+  Favorites = 'Favorites',
+}
+
+export enum SliderMode {
+  Infinity    = 'Infinity',
+  Sliders     = 'Sliders'
+}
+
+export enum Sensitivity {
+  x025 = '0.25',
+  x05 = '0.5',
+  x1 = '1',
+  x2 = '2',
+  x5 = '5',
+}
+
+type Action = {
+  timestamp: number;
+  target: string;
+  property: string;
+  value: any;
+}
+
+type Props = {
+};
+
+type State = {
+  selected: string;
+  actors: Record<string, string>;
+  values: Record<string, any>;
+  favorites: Record<string, boolean>;
+  mode: CCRMode;
+  value: any;
+  root: any;
+  colorProperty: ColorProperty;
+  range: { min: number, max: number };
+  search: string;
+  section: Section;
+  actorsGroup: ActorsGroup;
+  sliderMode: SliderMode;
+  sensitivity: Sensitivity;
+};
+
+export class ColorCorrection extends React.Component<Props, State> {
+
+  state: State = {
+    selected: '',
+    actors: {},
+    values: {},
+    favorites: {},
+    mode: CCRMode.Global,
+    value: {},
+    root: {},
+    colorProperty: ColorProperty.Gain,
+    range: { min: 0, max: 2 },
+    search: '',
+    section: Section.ColorCorrection,
+    actorsGroup: ActorsGroup.All,
+    sliderMode: SliderMode.Infinity,
+    sensitivity: Sensitivity.x1,
+  };
+
+  timer: NodeJS.Timeout;
+  actorsLength = 15;
+  updateSelection = false;
+  actionsStack: Action[] = [];
+  
+  async componentDidMount() {
+    try {
+      const favorites = localStorage.getItem('favorites');
+      if (favorites)
+        this.setState({ favorites: JSON.parse(favorites) });
+    } catch {
+    }
+
+    await this.onRefresh();
+    this.timer = setInterval(this.onRefresh, 1000);
+  }
+
+  componentDidUpdate(prevProps: Readonly<Props>, prevState: Readonly<State>) {
+    const { colorProperty, value, mode } = this.state;
+
+    if (colorProperty !== prevState.colorProperty || mode !== prevState.mode) {
+      const range = this.getRange(_.get(value, this.getColorProperty()));
+      this.setState({ range });
+    }
+  }
+
+  componentWillUnmount() {
+    clearInterval(this.timer);
+    this.timer = null;
+  }
+
+  onRefresh = async () => {
+    await this.onRefreshActors();
+    await this.onRefreshValues();
+  }
+
+  onActorSelect = (selected: string) => {
+    const { values } = this.state;
+    const value = values[selected] ?? {};
+    this.setState({ selected, value }, () => this.onRefreshValues(true));
+  }
+
+  onSpawnActor = async () => {
+    const { section } = this.state;
+    this.updateSelection = true;
+    let Class;
+
+    switch (section) {
+      case Section.ColorCorrection:
+        Class = '/ColorCorrectRegions/Blueprints/ColorCorrectRegion.ColorCorrectRegion_C';
+        break;
+
+      case Section.LightCards:
+        Class = '/Script/DisplayCluster.DisplayClusterLightCardActor';
+        break;
+
+      default:
+        return;
+    }
+
+    await _api.proxy.function(
+      '/Script/RemoteControlWebInterface.Default__RCWebInterfaceBlueprintLibrary',
+      'SpawnActor',
+      { Class }
+    );
+
+    await this.onRefreshActors();
+  }
+
+  onRefreshActors = async () => {
+    const { section } = this.state;
+    let Class;
+    switch (section) {
+      case Section.ColorCorrection:
+        Class = '/Script/ColorCorrectRegions.ColorCorrectRegion';
+        break;
+
+      case Section.LightCards:
+        Class = '/Script/DisplayCluster.DisplayClusterLightCardActor';
+        break;
+
+      default:
+        return;
+    }
+
+    const actors = await _api.proxy.function(
+      '/Script/RemoteControlWebInterface.Default__RCWebInterfaceBlueprintLibrary',
+      'FindAllActorsOfClass',
+      { Class }
+    ) ?? {};
+
+    const timestamp = Date.now();
+    const values = await _api.proxy.function(
+      '/Script/RemoteControlWebInterface.Default__RCWebInterfaceBlueprintLibrary',
+      'GetValuesOfActorsByClass',
+      { Class }
+    ) ?? {};
+
+    _.remove(this.actionsStack, action => action.timestamp <= timestamp);
+
+    for (const actor in values)
+      values[actor] = JSON.parse(values[actor]);
+
+    for (const action of this.actionsStack)
+      _.set(values[action.target], action.property, action.value);
+
+    if (!_.isEqual(this.state.actors, actors) || !_.isEqual(this.state.values, values)) {
+      let { selected } = this.state;
+
+      if (selected && !actors[selected])
+        selected = undefined;
+
+      if (!selected)
+        selected = _.first(Object.keys(actors));
+
+      if (this.updateSelection)
+        selected = _.last(Object.keys(actors));
+
+      const value = values[selected];
+
+      this.updateSelection = false;
+      this.setState({ actors, values, selected, value });
+    }
+  }
+
+  getRootComponentPath = () => {
+    const { selected } = this.state;
+    if (!selected)
+      return selected;
+
+    let index = selected.lastIndexOf("'");
+    if (index < 0)
+      return `${selected}.DefaultSceneRoot`;
+
+    return selected.slice(0, index) + '.DefaultSceneRoot' + selected.slice(index);
+  }
+
+  onRefreshValues = async (update = false) => {
+    let { selected, range } = this.state;
+    if (!selected)
+      return;
+
+    //const value = await _api.proxy.put('/remote/object/property', { objectPath: selected, access: 'READ_ACCESS' });
+    const root = await _api.proxy.put('/remote/object/property', { objectPath: this.getRootComponentPath(), access: 'READ_ACCESS' });
+
+    if (!_.isEqual(this.state.root, root)) {
+      // if (update || _.isEmpty(this.state.value))
+      //   range = this.getRange(_.get(value, this.getColorProperty()));
+      this.setState({ /*value,*/ root, range });
+    }
+  }
+
+  onRootPropertyValueChangeInternal = async (propertyName: string, property: string, value: any) => {
+    const { selected, root } = this.state;
+    if (!selected)
+      return;
+
+    const propertyValue = _.set({}, _.compact([propertyName, property]), value);
+    _.merge(root, propertyValue);
+    this.setState({ root });
+
+    await _api.proxy.put('/remote/object/property', {
+      objectPath: this.getRootComponentPath(),
+      access: 'WRITE_TRANSACTION_ACCESS',
+      propertyName,
+      propertyValue
+    });
+  }
+
+  onPropertyValueChangeInternal = async (propertyName: string, property: string, value: any) => {
+    const { selected } = this.state;
+    if (!selected)
+      return;
+
+    const propertyValue = _.set({}, property, value);
+    const action: Action = { timestamp: Date.now(), target: selected, property, value };
+
+    this.actionsStack.push(action);
+
+    await _api.proxy.put('/remote/object/property', {
+      objectPath: selected,
+      access: 'WRITE_TRANSACTION_ACCESS',
+      propertyName,
+      propertyValue
+    });
+  }
+
+  onThrottledPropertyValueChange = _.throttle(this.onPropertyValueChangeInternal, 100);
+  onThrottledRootPropertyValueChange = _.throttle(this.onRootPropertyValueChangeInternal, 100);
+
+  onPropertyChange = (property: string, v: any) => {
+    const { value } = this.state;
+    _.set(value, property, v);
+    this.setState({ value });
+
+    const [propertyName] = property.split('.');
+    this.onThrottledPropertyValueChange(propertyName, property, v);
+  }
+
+  onReset = async (property: string, value: PropertyValue) => {
+    if (!await AlertModal.show('Are you sure you want to reset?'))
+      return;
+
+    this.onPropertyChange(property, value);  
+  }
+
+  onFavoriteToggle = () => {
+    const { selected, favorites, actorsGroup } = this.state;
+    const actor = favorites[selected];
+
+    if (actorsGroup === ActorsGroup.Favorites && !actor)
+      return;
+
+    favorites[selected] = !favorites[selected];
+    if (!favorites[selected])
+      delete favorites[selected];
+
+    localStorage.setItem('favorites', JSON.stringify(favorites));
+    
+    this.setState({ favorites });
+  }
+
+  onSectionChange = (section: Section) => {
+    let { mode } = this.state;
+
+    if (section === Section.ColorCorrection && mode === CCRMode.Color)
+      mode = CCRMode.Global;
+
+    this.setState({ section, mode, actors: {}, selected: undefined, root: undefined }, this.onRefresh);
+  }
+
+  onInfiniteWheelMove = (value: number, sign: number, min: number, max: number, property: string) => {
+    const { sensitivity } = this.state;
+    const step = sign * 0.0025 * (max - min) * +sensitivity;
+
+    value = Math.max(min, Math.min(max, value + step));
+    this.onPropertyChange(property, value);
+  }
+
+  getRange = (color) => {
+    const { mode, colorProperty } = this.state;
+    if (mode === CCRMode.Color)
+      return ({ min: 0, max: 1 });
+
+    let min = 0;
+    let max = 2;
+
+    if (colorProperty === ColorProperty.Offset) {
+      min = -1;
+      max = 1;
+    }
+
+    for (const key in color) {
+      if (color[key] > max)
+        max = color[key];
+    }
+
+    return { min, max };
+  }
+
+  getColorProperty = () => {
+    const { section, mode, colorProperty } = this.state;
+
+    if (mode === CCRMode.Color)
+      return 'Color';
+
+    let root = '';
+    switch (section) {
+      case Section.ColorCorrection:
+        root = 'ColorGradingSettings';
+        break;
+
+      case Section.LightCards:
+        root = 'color grading';
+        break;
+    }
+
+    return `${root}.${mode}.${colorProperty}`;
+  }
+
+  renderColorSlider = (property: string, min: number, max: number, reset: number, type: string) => {
+    return this.renderSlider(property, min, max, reset, 2, null, false, type);
+  }
+
+  renderSlider = (property: string, min: number, max: number, reset?: number, precision: number = 2, label?: string, disabled?: boolean, type: string = '') => {
+    const { value, sliderMode } = this.state;
+    const v = _.get(value, property);
+
+    let className = `control wrap ${property.replace(' ', '-').replace('.', '-')}-slider `;
+
+    if (disabled)
+      className += 'disabled ';
+
+    let slider = <Slider showLabel={true}
+                         precision={precision}
+                         value={v}
+                         onChange={value => this.onPropertyChange(property, value)}
+                         min={min}
+                         max={max} />;
+
+    if (sliderMode === SliderMode.Infinity)
+      slider = (
+        <div className="infinite-slider-container">
+          <SliderWheel size={200}
+                       className={`infinite-slider ${type}`}
+                       onWheelMove={sign => this.onInfiniteWheelMove(v, sign, min, max, property)} />
+          <ValueInput draggable={false}
+                      precision={precision}
+                      min={min}
+                      max={max}
+                      value={v}
+                      onChange={value => this.onPropertyChange(property, value)} />
+        </div>
+      );
+
+
+    return (
+      <div className={className} key={property}>
+        <label>{label ?? _.startCase(_.last(property.split('.')))}</label>
+        {slider}
+        <FontAwesomeIcon icon={['fas', 'undo']} onClick={() => this.onPropertyChange(property, reset)} />
+      </div>
+    );
+  }
+
+  renderToggle = (property: string, label?: string) => {
+    const { value } = this.state;
+
+    return (
+      <div className={`prop ${property.replaceAll(' ', '-').replaceAll('.', '-')}-toggle`}>
+        <label>{label ?? _.startCase(property)}</label>
+        <div className={`boolean-toggle ${(!_.get(value, property)) && 'checked'} `}>
+          <label className="switch toggle-mode">
+            <input type="checkbox" checked={!_.get(value, property)} onChange={e => this.onPropertyChange(property, !e.target.checked)} />
+            <span className="slider inline"></span>
+          </label>
+          <div className='labels'>
+            <div className='off'>ON</div>
+            <div className='on'>OFF</div>
+          </div>
+        </div>
+      </div>
+    );
+  }
+
+  renderSelect = (property: string, options: string[] | Record<string, string>) => {
+    const { value } = this.state;
+
+    let opt = [];
+    if (Array.isArray(options)) {
+      opt = options.map(value => ({ value, label: value }));
+    } else {
+      for (const value in options)
+        opt.push({ value, label: options[value] });
+    }
+
+    return (
+      <div className='prop'>
+        <label>{_.startCase(property)}</label>
+        <select className='dropdown' 
+                value={_.get(value, property)}
+                onChange={e => this.onPropertyChange(property, e.target.value)}>
+          {opt.map(option => <option value={option.value} key={option.value}>{option.label}</option>)}
+        </select>
+      </div>
+    );
+  }
+
+  executeButton = (func: string) => {
+    const { selected } = this.state;
+    return _api.proxy.function(selected, func, {});
+  }
+
+  renderButton = (func: string) => {
+    return (
+      <div className="prop">
+        <Button label={func} onExecute={() => this.executeButton(func)} />
+      </div>
+    );
+  }
+
+  renderGuidesColorPicker = () => {
+    const { value } = this.state;
+
+    const color = _.get(value, 'GuideColor');
+
+    return (
+      <div className="guides-color-picker">
+        <div className="label">Guide Color</div>
+        <div className="color-controls-row">
+          <ColorPicker value={color}
+                       reset={true}
+                       mode={ColorMode.Rgb}
+                       type={PropertyType.LinearColor}
+                       onChange={value => this.onPropertyChange('GuideColor', value ?? { R: 1, G: 1, B: 0, A: 0 })} />
+          <div className="control-rows">
+            {this.renderColorSlider('GuideColor.R', 0, 1, 1, 'red-slider')}
+            {this.renderColorSlider('GuideColor.G', 0, 1, 1, 'green-slider')}
+            {this.renderColorSlider('GuideColor.B', 0, 1, 1, 'blue-slider')}
+          </div>
+        </div>
+      </div>
+    );
+  }
+
+  isColorModeModified = (actor: string, mode: CCRMode) => {
+    const { values, section } = this.state;
+
+    const grading = _.get(values, [actor, (section === Section.ColorCorrection ? 'ColorGradingSettings' : 'color grading'), mode], {});
+    for (const color in grading) {
+      const def = (color === ColorProperty.Offset ? 0 : 1);
+      for (const property in grading[color]) {
+        const v = grading[color][property];
+        if (v !== def)
+          return true;
+      }
+    }
+
+    return false;
+  }
+
+  renderChangeIndication = (actor: string, mode: CCRMode) => {
+    let className = 'indicator ';
+    if (this.isColorModeModified(actor, mode))
+      className += 'active ';
+
+    return (
+      <div className={className}>
+        <FontAwesomeIcon icon={['fas', 'circle']} />
+      </div>
+    );
+  }
+
+  renderTabsIcons = () => {
+    const { favorites, selected } = this.state;
+
+    let icon: IconProp = ['far', 'star'];
+    if (favorites[selected])
+      icon = ['fas', 'star'];
+
+    return <>
+      <span className='tab-icon' onClick={this.onFavoriteToggle}><FontAwesomeIcon icon={icon} /></span>
+      {/* <span className='tab-icon'><FontAwesomeIcon icon={['fas', 'crosshairs']} /></span> */}
+    </>;
+  }
+
+  renderActorsContent = (pagination: number, items: number, actors: Record<string, any>, favorite: boolean) => {
+    const { selected, search, favorites, section, values } = this.state;
+
+    let start = pagination * items;
+    let end = start + items;
+ 
+    if (favorite)
+      actors = Object.fromEntries(_.toPairs(actors).filter(([key, value]) => value));
+
+    let actorsKeys = Object.keys(actors);
+    if (search)
+      actorsKeys = actorsKeys.filter(actor => this.state.actors[actor].toLowerCase().includes(search));
+
+    return (
+      <>
+        {actorsKeys.slice(start, end).map(actor =>
+          <div key={actor}
+               className={`light-card ${selected === actor ? 'active' : ''}`}
+               onClick={() => this.onActorSelect(actor)}>
+            {_.startCase(this.state.actors[actor])}
+            { section === Section.LightCards &&
+              <div className={`power-window-icon ${_.get(values, [actor, 'Power Window Mode']) ? 'active' : ''}`}>
+                <svg width="28" height="29" viewBox="0 0 28 29" fill="none" xmlns="http://www.w3.org/2000/svg">
+                  <path d="M21.7 6.77118V23.221L4.55 20.2461V9.74609L21.7 6.77118ZM24.5 4.49609L3.5 8.87109V21.1211L24.5 25.4961V4.49609Z" fill="white" fillOpacity="0.35"/>
+                  <path opacity="0.5" d="M19.075 20.2461L7 18.1462V11.4961L19.075 20.2461Z" fill="white" fillOpacity="0.35"/>
+                  <path fillRule="evenodd" clipRule="evenodd" d="M19.25 15.5219L12.075 10.4468L19.25 9.22168V15.5219Z" fill="white" fillOpacity="0.35"/>
+                  <path opacity="0.5" d="M1.75 2.74609H7V0.996094H0V7.99609H1.75V2.74609Z" fill="white" fillOpacity="0.35"/>
+                  <path opacity="0.5" d="M26.25 2.74609V7.99609H28V0.996094H21V2.74609H26.25Z" fill="white" fillOpacity="0.35"/>
+                  <path opacity="0.5" d="M26.25 27.2461H21V28.9961H28V21.9961H26.25V27.2461Z" fill="white" fillOpacity="0.35"/>
+                  <path opacity="0.5" d="M1.75 27.2461V21.9961H0V28.9961H7V27.2461H1.75Z" fill="white" fillOpacity="0.35"/>
+                </svg>
+              </div>
+            }
+            {!!favorites[actor] && 
+              <FontAwesomeIcon className="icon-favorite" icon={['fas', 'star']} />
+            }
+
+            <div className='change-indicators'>
+              {this.renderChangeIndication(actor, CCRMode.Global)}
+              {this.renderChangeIndication(actor, CCRMode.Highlights)}
+              {this.renderChangeIndication(actor, CCRMode.Midtones)}
+              {this.renderChangeIndication(actor, CCRMode.Shadows)}
+            </div>
+          </div>
+        )}
+        {!favorite && (
+          <div className="light-card add-button"
+               onClick={this.onSpawnActor}>
+            <FontAwesomeIcon icon={['fas', 'plus']} />
+          </div>
+        )}
+      </>
+    );
+  }
+
+  renderAppearanceContent = (pagination: number): Element[] => {
+    const contents = [];
+
+    if (pagination === 0)
+      contents.push(
+        this.renderSlider('Temperature', 0, 10_000, 6_500),
+        this.renderSlider('Tint', -1, 1, 0, 3),
+        this.renderSlider('Exposure', -2, 10, 0.5, 3, 'Exposure'),
+        this.renderSlider('Intensity', -2, 10, 1, 3, null, _.get(this.state.value, 'Power Window Mode', false)),
+      );
+
+    if (pagination === 1)
+      contents.push(
+        this.renderSlider('Opacity', 0, 1, 1, 3),
+        this.renderSlider('Feathering', 0, 1, 0, 3),
+      );
+
+    return contents;
+  }
+
+  renderOrientationContent = (pagination: number) => {
+    const { value } = this.state;
+    const contents = [];
+
+    if (pagination === 0)
+      contents.push(
+        this.renderSlider('DistanceFromCenter', -2_000, 2_000, 300),
+        this.renderSlider('Longitude', 0, 360, 0, 2, null, _.get(value, 'UseClockSystem')),
+        this.renderSlider('Latitude', -90, 90, 0),
+        this.renderSlider('Spin', 0, 360, 0),
+      );
+
+    if (pagination === 1)
+      contents.push(
+        this.renderSlider('Scale.X', 0, 20, 2, 3, 'Scale X'),
+        this.renderSlider('Scale.Y', 0, 20, 2, 3, 'Scale Y'),
+        this.renderSlider('WallIndex', -1, 20, -1, 0),
+      );
+
+    if (pagination === 2)
+      contents.push(
+        this.renderSlider('Unit', 0, 200, 0),
+        <div className="buttons-layout" key="Updown">
+          <div className="buttons-content">
+            {this.renderButton('Move Up')}
+            {this.renderButton('Move Down')}
+          </div>
+          <div className="buttons-content small-button"  key="Sides">
+            {this.renderButton('Move X+')}
+            <div className="space-buttons">
+              {this.renderButton('Move Y-')}
+              {this.renderButton('Move Y+')}
+            </div>
+            {this.renderButton('Move X-')}
+          </div>
+        </div>
+      );
+
+    return contents;
+  }
+
+  renderPositionTab = () => {
+    const { root, sensitivity } = this.state;
+
+    return (
+      <>
+        <div className="control-rows no-space vector-controls">
+          <div className="control">
+            <VectorControl value={_.get(root, 'RelativeLocation')}
+                           label="Location"
+                           propertyType={PropertyType.Vector}
+                           step={+sensitivity}
+                           defaultValue={{ X: 0, Y: 0, Z: 0}}
+                           onChange={(property, value) => this.onThrottledRootPropertyValueChange('RelativeLocation', property, value)} />
+          </div>
+          <div className="control">
+            <VectorControl value={_.get(root, 'RelativeRotation')}
+                           label="Rotation"
+                           propertyType={PropertyType.Rotator}
+                           step={+sensitivity}
+                           defaultValue={{ Pitch: 0, Roll: 0, Yaw: 0}}
+                           onChange={(property, value) => this.onThrottledRootPropertyValueChange('RelativeRotation', property, value)} />
+          </div>
+          <div className="control">
+            <VectorControl value={_.get(root, 'RelativeScale3D')}
+                           label="Scale"
+                           propertyType={PropertyType.Vector}
+                           step={+sensitivity}
+                           defaultValue={{ X: 1, Y: 1, Z: 1}}
+                           onChange={(property, value) => this.onThrottledRootPropertyValueChange('RelativeScale3D', property, value)} />
+          </div>
+        </div>
+      </>
+    );
+  }
+
+  renderSection = () => {
+    const { section, value } = this.state;
+
+    switch (section) {
+      case Section.ColorCorrection:
+        return (
+          <Tabs rightIcon={this.renderTabsIcons()}
+                defaultActiveKey="default">
+            <TabPane id="position" tab="Position">
+              {this.renderPositionTab()}
+            </TabPane>
+            <TabPane id="default" tab="Default">
+              <div className='props-row'>
+                {this.renderSelect('Type', ['Sphere', 'Box', 'Cylinder', 'Cone'])}
+                {this.renderToggle('Enabled')}
+              </div>
+              <div className='control-rows'>
+                {this.renderSlider('Temperature', 0, 10_000, 6_500)}
+                {this.renderSlider('Intensity', 0, 1, 1)}
+                {this.renderSlider('Falloff', 0, 20, 0, 3)}
+                {this.renderSlider('Inner', 0, 1, 0.5)}
+              </div>
+            </TabPane>
+            <TabPane id="advanced" tab="Advanced">
+              <div className='props-row'>
+                {this.renderToggle('Invert')}
+              </div>
+              <div className='control-rows'>
+                {this.renderSlider('ColorGradingSettings.ShadowsMax', 0, 1, 0)}
+                {this.renderSlider('ColorGradingSettings.HighlightsMin', 0, 1, 0)}
+                {this.renderSlider('Outer', 0, 1, 1, 3)}
+                {this.renderSlider('Priority', -1, 20, 0, 0)}
+              </div>
+            </TabPane>
+          </Tabs>
+        );
+
+      case Section.LightCards:
+        return (
+          <Tabs rightIcon={this.renderTabsIcons()}
+                defaultActiveKey="appearance">
+            <TabPane id="appearance" tab="Appearance">
+              <div className="grid-row">
+                <div className="props-row">
+                  {this.renderSelect('Mask', { NewEnumerator1: 'Square', NewEnumerator0: 'Circle' })}
+                  <div className="prop asset-control">
+                    <AssetWidget browse
+                                 reset={false}
+                                 type="Texture2D"
+                                 label="Texture"
+                                 onChange={v => this.onPropertyChange('Texture', v)}
+                                 value={_.get(value, 'Texture') ?? ''} />
+                    <span className="reset"><FontAwesomeIcon icon={['fas', 'undo']} onClick={() => this.onPropertyChange('Texture', null)} /></span>
+                  </div>
+                </div>
+                <div className="props-row">
+                  {this.renderToggle('Power Window Mode', 'Power Window')}
+                  {this.renderToggle('useCCmult', 'CC Multiple')}
+                </div>
+              </div>
+              <div className="multilines-content">
+                <PaginationContent pages={2} key="appearance">
+                  {({ pagination }) => this.renderAppearanceContent(pagination)}
+                </PaginationContent>
+              </div>
+            </TabPane>
+            <TabPane id="orientation" tab="Orientation">
+              <div className="props-row grid-row orientation">
+                <div className="prop-item">
+                  {this.renderToggle('UseClockSystem', 'Clock System')}
+                  <div className="prop">
+                    <TimeControl label="Time"
+                                  value={_.get(value, 'Hour')}
+                                  disabled={!_.get(value, 'UseClockSystem')}
+                                  onChange={v => this.onPropertyChange('Hour', v)} />
+                  </div>
+                </div>
+                <div className="prop-item">
+                  {this.renderToggle('StickToWall')}
+                  {this.renderButton('Set To Center of Stage')}
+                </div>
+              </div>
+              <div className="multilines-content">
+                <PaginationContent pages={3} key="orientation">
+                  {({ pagination }) => this.renderOrientationContent(pagination)}
+                </PaginationContent>
+              </div>
+            </TabPane>
+            <TabPane id="guides" tab="Guides">
+              <div className='props-row'>
+                {this.renderToggle('bShowGuides', 'Show Guides')}
+              </div>
+              <div className='control-rows'>
+                {this.renderSlider('GuideScale', 0, 1, 0)}
+                {this.renderGuidesColorPicker()}
+              </div>
+            </TabPane>
+            <TabPane id="position" tab="Position">
+              {this.renderPositionTab()}
+            </TabPane>
+            <TabPane id="advanced" tab="Advanced">
+              <div className='control-rows light-advanced'>
+                {this.renderSlider('color grading.ShadowsMax', -1, 1, 0.09)}
+                {this.renderSlider('color grading.HighlightsMin', -1, 1, 0.5)}
+              </div>
+            </TabPane>
+          </Tabs>
+        );
+    }
+
+    return null;
+  }
+
+  render() {
+    const { selected, mode, actors, favorites, actorsGroup, section, value, colorProperty, range, sliderMode, sensitivity } = this.state;
+    const style: React.CSSProperties = {};
+    const property = this.getColorProperty();
+    const color = _.get(value, property);
+
+    style.position = 'relative';
+
+    let className = 'cc-group cc-panel ';
+
+    if (!selected)
+      className += 'disabled';
+
+    let availableFavorites = _.pickBy(favorites, (v, k) => !!actors[k]);
+
+    return (
+      <div className="color-correction-wrapper custom-ui-wrapper" tabIndex={-1}>
+        <div className='ipad-debug'></div>
+        <img src='/images/ref001.jpg' className='ref'/>
+        <div className="group">
+          <CorrectionColorPicker section={section}
+                                 property={property}
+                                 color={color}
+                                 colorProperty={colorProperty}
+                                 range={range}
+                                 mode={mode}
+                                 disabled={!selected}
+                                 sliderMode={sliderMode}
+                                 sensitivity={sensitivity}
+                                 onSectionChange={this.onSectionChange}
+                                 onPropertyChange={this.onPropertyChange}
+                                 onRangeUpdate={color => this.setState({ range: this.getRange(color) })}
+                                 onColorPropertyChange={colorProperty => this.setState({ colorProperty })}
+                                 onColorModeChange={mode => this.setState({ mode })}
+                                 isColorPropertyModified={mode => this.isColorModeModified(selected, mode)} />
+        </div>
+        <div className="group">
+          <div className="cc-group shotbox cc-panel">
+            <Tabs rightIcon={<Search placeholder="Search" onSearch={search => this.setState({ search })} />}
+                  onTabChange={(actorsGroup: ActorsGroup) => this.setState({ actorsGroup })}
+                  defaultActiveKey={actorsGroup}>
+              <TabPane id={ActorsGroup.All} tab="All" icon={['fas', 'square']}>
+                <PaginationContent pageItems={this.actorsLength}
+                                   total={Object.values(actors).length}
+                                   key="all">
+                  {({ pagination }) => this.renderActorsContent(pagination, this.actorsLength, actors, false)}
+                </PaginationContent>
+              </TabPane>
+              <TabPane id={ActorsGroup.Favorites} tab="Favorites" icon={['fas', 'star']}>
+                <PaginationContent pageItems={this.actorsLength + 1}
+                                   total={Object.values(availableFavorites).length}
+                                   key="favorites">
+                  {({ pagination }) => this.renderActorsContent(pagination, this.actorsLength, availableFavorites, true)}
+                </PaginationContent>
+              </TabPane>
+            </Tabs>
+          </div>
+          <div className={className} style={style}>
+            <div className="body">
+              {this.renderSection()}
+            </div>
+          </div>
+        </div>
+        <div className="group bottom-bar">
+          <div className="cc-group inline-group">
+            <div className="group-item operation-style bottom-tabs-style">
+              <div className="label">
+                Operation Style
+              </div>
+              <Tabs onlyHeader
+                    defaultActiveKey={sliderMode}
+                    onTabChange={(sliderMode: SliderMode) => this.setState({ sliderMode })}>
+                {Object.values(SliderMode).map(mode => <TabPane key={mode}
+                                                                tab={mode}
+                                                                id={mode} />)}
+              </Tabs>
+            </div>
+            <div className="group-item sensitivity-style bottom-tabs-style">
+              <div className="label">
+                Sensitivity
+              </div>
+              <Tabs onlyHeader
+                    defaultActiveKey={sensitivity}
+                    onTabChange={(sensitivity: Sensitivity) => this.setState({ sensitivity })}>
+                {Object.values(Sensitivity).map(mode => <TabPane key={mode}
+                                                                 tab={`x${mode}`}
+                                                                 id={mode} />)}
+              </Tabs>
+            </div>
+          </div>
+        </div>
+      </div>
+    );
+  }
+}
diff --git a/Engine/Plugins/VirtualProduction/RemoteControlWebInterface/WebApp/Client/src/components/Screens/ColorCorrection/CorrectionColorPicker.tsx b/Engine/Plugins/VirtualProduction/RemoteControlWebInterface/WebApp/Client/src/components/Screens/ColorCorrection/CorrectionColorPicker.tsx
new file mode 100644
--- /dev/null
+++ b/Engine/Plugins/VirtualProduction/RemoteControlWebInterface/WebApp/Client/src/components/Screens/ColorCorrection/CorrectionColorPicker.tsx
@@ -0,0 +1,360 @@
+import React from 'react';
+import { ReactComponent as CCIcon } from '../../../assets/cc_icon.svg';
+import { ReactComponent as LCIcon } from '../../../assets/lc_icon.svg';
+import { ColorPicker, SliderWheel, TabPane, Tabs, ValueInput } from '../../controls';
+import { FontAwesomeIcon } from '@fortawesome/react-fontawesome';
+import { Section, ColorProperty, CorrectionColorSlider } from './';
+import { AlertModal } from 'src/components/modals';
+import { WidgetUtilities } from 'src/utilities';
+import { CCRMode, ColorMode, Sensitivity, SliderMode } from './ColorCorrection';
+import { PropertyType, VectorProperty } from 'src/shared';
+import _ from 'lodash';
+
+
+type Props = {
+  section: Section;
+  property: string;
+  color: any;
+  colorProperty: ColorProperty;
+  range: { min: number, max: number };
+  mode: CCRMode;
+  sliderMode: SliderMode;
+  sensitivity: Sensitivity;
+  disabled: boolean;
+
+  onSectionChange: (section: Section) => void;
+  onColorPropertyChange: (colorProperty: ColorProperty) => void;
+  onPropertyChange: (property: string, v: any) => void;
+  onRangeUpdate: (color: any) => void;
+  onColorModeChange: (mode: CCRMode) => void;
+  isColorPropertyModified: (mode: CCRMode) => boolean;
+}
+
+type State = {
+  colorMode: ColorMode,
+}
+
+export class CorrectionColorPicker extends React.Component<Props, State> {
+
+  state: State = {
+    colorMode: ColorMode.Rgb,
+  }
+
+  onReset = async (property: string, value) => {
+    if (!await AlertModal.show('Are you sure you want to reset?'))
+      return;
+
+    this.props.onPropertyChange(property, value);
+  }
+
+  onColorModeChange = () => {
+    let colorMode = ColorMode.Hsv;
+
+    if (this.state.colorMode === ColorMode.Hsv)
+      colorMode = ColorMode.Rgb;
+    
+    this.setState({ colorMode });
+  }
+
+  onColorChange = (v: any) => {
+    const { color, property } = this.props;
+
+    if (_.isObject(color))
+      v = { ...color, ...v };
+
+    return this.props.onPropertyChange(property, v);
+  }
+
+  onHsvWheelMove = (sign: number, hue?: number) => {
+    let { range, color, sensitivity } = this.props;
+    let rgb = WidgetUtilities.colorToRgb(color, range.max);
+
+    const hsv = WidgetUtilities.rgb2Hsv(rgb);
+
+    if (hue)
+      hsv.h = hue / 360;
+    else {
+      let step = sign * 0.001 * +sensitivity;
+      if (hsv.h + step < 0)
+        step += 1;
+
+      hsv.h = hsv.h + step;
+    }
+
+    rgb = WidgetUtilities.hsv2rgb(hsv);
+    color = WidgetUtilities.rgbToColor(rgb, range.max);
+
+    this.onColorChange(color);
+  }
+  
+  onShiftSliderMove = (sign: number) => {
+    const { colorProperty, property, color, sensitivity } = this.props;
+    const step = sign * 0.01 * +sensitivity;
+
+    if (!color)
+      return;
+
+    let sum = color.X + color.Y + color.Z;
+    let min = 0;
+
+    if (colorProperty === ColorProperty.Offset)
+      min = -1;
+
+    color.X = Math.max(min, color.X + ((color.X / sum) || 1) * step);
+    color.Y = Math.max(min, color.Y + ((color.Y / sum) || 1) * step);
+    color.Z = Math.max(min, color.Z + ((color.Z / sum) || 1) * step);
+
+    this.props.onPropertyChange(property, color);
+    this.props.onRangeUpdate(color);
+  }
+
+  getSectionProps = (value: Section) => {
+    const { section, onSectionChange } = this.props;
+    let className = 'cc-group mode-button ';
+
+    if (section === value)
+      className += 'active ';
+
+    return { className, onClick: () => onSectionChange(value) };
+  }
+
+  getColorResetValue = () => {
+    const { mode, colorProperty } = this.props;
+    if (mode === CCRMode.Color)
+      return ({ R: 1, G: 1, B: 1, A: 0});
+
+    if (colorProperty === ColorProperty.Offset)
+      return ({ X: 0, Y: 0, Z: 0, W: 0 });
+
+    return ({ X: 1, Y: 1, Z: 1, W: 1 });
+  }
+
+  getHsvValue = () => {
+    const { range, color } = this.props;
+
+    const rgb = WidgetUtilities.colorToRgb(color, range.max);
+    const hsv = WidgetUtilities.rgb2Hsv(rgb);
+
+    return hsv;
+  }
+
+  getValueColor = (): React.CSSProperties => {
+    const { color, mode } = this.props;
+
+    if (!color || mode === CCRMode.Color)
+      return { opacity: 0 };
+
+    const max = Math.max(color.X, color.Y, color.Z, 1);
+
+    const rgb = WidgetUtilities.colorToRgb(color, max);
+    const hsv = WidgetUtilities.rgb2Hsv(rgb);
+
+    let opacity = 0;
+    if (hsv.v)
+      opacity = Math.max(0, Math.min(1, +(1 - hsv.v).toFixed(4)));
+
+    return { opacity };
+  }
+
+  getPreviewColorStyles = (color: VectorProperty): React.CSSProperties => {
+    const { range } = this.props;
+    const rgb = WidgetUtilities.colorToRgb(color, range.max);
+
+    return { background: `rgb(${Math.round(rgb.R)},${Math.round(rgb.G)},${Math.round(rgb.B)})` };
+  }
+
+  getPropertyType = () => {
+    const { mode } = this.props;
+    let propertyType = PropertyType.Vector4;
+
+    if (mode === CCRMode.Color)
+      propertyType = PropertyType.LinearColor;
+
+    return propertyType;
+  }
+
+  renderColorSlider = (sliderProperty: string, label: string, className: string) => {
+    const { colorMode } = this.state;
+    const { color, colorProperty, property, range, sliderMode, sensitivity } = this.props;
+    const propertyType = this.getPropertyType();
+
+    return <CorrectionColorSlider className={className}
+                                  label={label}
+                                  mode={sliderMode}
+                                  sensitivity={+sensitivity}
+                                  color={color}
+                                  colorProperty={colorProperty}
+                                  property={property}
+                                  sliderProperty={sliderProperty}
+                                  range={range}
+                                  colorMode={colorMode}
+                                  propertyType={propertyType}
+                                  onPropertyChange={this.props.onPropertyChange}
+                                  onRangeUpdate={this.props.onRangeUpdate} />;
+  }
+
+  renderSliders = () => {
+    const propertyType = this.getPropertyType();
+
+    if (propertyType === PropertyType.LinearColor) {
+      return (
+        <>
+          {this.renderColorSlider('R', 'R', 'red-slider')}
+          {this.renderColorSlider('G', 'G', 'green-slider')}
+          {this.renderColorSlider('B', 'B', 'blue-slider')}
+          {this.renderColorSlider('A', 'A', 'alpha-slider')}      
+        </>
+      );
+    }
+    
+    if (this.state.colorMode === ColorMode.Hsv) {
+      return (
+        <>
+          {this.renderColorSlider('s', 'S', 'saturation-slider')}
+          {this.renderColorSlider('v', 'V', 'alpha-slider')}
+          {this.renderColorSlider('a', 'Y', 'alpha-slider')}
+        </>
+      );
+    }
+
+    return (
+      <>
+        {this.renderColorSlider('X', 'R', 'red-slider')}
+        {this.renderColorSlider('Y', 'G', 'green-slider')}
+        {this.renderColorSlider('Z', 'B', 'blue-slider')}
+        {this.renderColorSlider('W', 'Y', 'alpha-slider')}      
+      </>
+    );
+  }
+
+  renderModifyIndicator = (mode: CCRMode) => {
+    if (!this.props.isColorPropertyModified?.(mode))
+      return null;
+
+    return <div className="change-indicator active"><FontAwesomeIcon icon={['fas', 'circle']} /></div>;
+  }
+
+  render() {
+    const { property, color, colorProperty, range, mode, disabled, section } = this.props;
+    const { colorMode } = this.state;
+
+    const reset = this.getColorResetValue();
+    const hsv = this.getHsvValue();
+
+    const hue = hsv ? (hsv.h * 360).toFixed(1) : 0;
+    const style: React.CSSProperties = {};
+
+    if (disabled) {
+      style.opacity = 0.6;
+      style.pointerEvents = 'none';
+    }
+
+    let modes = Object.values(CCRMode);
+    if (section === Section.ColorCorrection)
+      modes = modes.filter(m => m !== CCRMode.Color);
+
+    return (
+      <div className="group">
+        <div className="cc-group screen-mode-selectors">
+          <div className="mode-selector">
+            <div {...this.getSectionProps(Section.ColorCorrection)}>
+              <CCIcon />
+            </div>
+            <div {...this.getSectionProps(Section.LightCards)}>
+              <LCIcon />
+            </div>
+          </div>
+        </div>
+        <div className="cc-group" style={style}>
+          <div className="body" style={{ width: "calc(100% - 20px)" }}>
+            <div className="mode-selectors">
+              <Tabs onTabChange={this.props.onColorModeChange}
+                    onlyHeader
+                    defaultActiveKey={mode}>
+                {modes.map(mode => <TabPane key={mode}
+                                            id={mode} 
+                                            tab={mode}
+                                            view={() => this.renderModifyIndicator(mode)} />
+                )}
+
+              </Tabs>
+            </div>
+            <div className="color-picker-singleton">
+              <div className="color-wheel">
+                <FontAwesomeIcon className="reset-icon" icon={['fas', 'undo']} onClick={() => this.onReset(property, reset)} />
+                {colorMode === ColorMode.Rgb && <div className="color-picker-value" style={this.getValueColor()} />}
+                <ColorPicker value={color}
+                             alpha={mode === CCRMode.Color}
+                             mode={colorMode}
+                             type={this.getPropertyType()}
+                             max={range.max}
+                             onChange={value => this.onColorChange(value ?? reset)} />
+                <div className="wheel-addon">
+                  <div className="dot" style={this.getPreviewColorStyles(color)}></div>
+                  <div className={`color-mode-toggle ${(colorMode === ColorMode.Hsv) && 'checked'} `}>
+                    <label className="switch toggle-mode">
+                      <input type="checkbox" checked={colorMode === ColorMode.Hsv} onChange={this.onColorModeChange} />
+                      <span className="slider inline"></span>
+                    </label>
+                    <div className="labels">
+                      <div className="off">RGB</div>
+                      <div className="on">HSV</div>
+                    </div>
+                  </div>
+                </div>
+              </div>
+              <div className="shift-slider">
+                {colorMode === ColorMode.Rgb && <>
+                  <SliderWheel vertical
+                               onWheelMove={this.onShiftSliderMove} />
+                  Value
+                </>
+                }
+              </div>
+              <div className="details-section">
+                <div className="color-property-selection">
+                  {mode !== CCRMode.Color && Object.values(ColorProperty).map(property =>
+                    <div key={property}
+                         onClick={() => this.props.onColorPropertyChange(property)}
+                         className={`mode-selector ${property === colorProperty ? 'active' : ''}`}>{property}</div>
+                  )}
+                </div>
+                <div className="color-properties">
+                  <div className="sliders">
+                    <div className="rgb-sliders">
+                      {this.renderSliders()}
+                      <div className="slider-wheel-container hue-slider">
+                        <span className="title">Hue</span>
+                        <SliderWheel size={200} onWheelMove={sign => this.onHsvWheelMove(sign)} />
+                        <ValueInput draggable={false}
+                                    precision={1}
+                                    min={0}
+                                    max={360}
+                                    value={hue}
+                                    onChange={hue => this.onHsvWheelMove(1, hue)} />
+                      </div>
+                      {(mode === CCRMode.Shadows) &&
+                        <div className="slider-wheel-container shadow-min-slider">
+                          <span className="title">Shadow Min</span>
+                          <SliderWheel size={200}
+                                      className="infinite-slider" 
+                                      onWheelMove={sign => {}} />
+                          <ValueInput draggable={false}
+                                      precision={1}
+                                      min={0}
+                                      max={360}
+                                      value={0}
+                                      onChange={v => {}} />
+                        </div>
+                      }
+                    </div>
+                  </div>
+                </div>
+              </div>
+            </div>
+          </div>
+        </div>
+      </div>
+    );
+  }
+};
\ No newline at end of file
diff --git a/Engine/Plugins/VirtualProduction/RemoteControlWebInterface/WebApp/Client/src/components/Screens/ColorCorrection/CorrectionColorSlider.tsx b/Engine/Plugins/VirtualProduction/RemoteControlWebInterface/WebApp/Client/src/components/Screens/ColorCorrection/CorrectionColorSlider.tsx
new file mode 100644
--- /dev/null
+++ b/Engine/Plugins/VirtualProduction/RemoteControlWebInterface/WebApp/Client/src/components/Screens/ColorCorrection/CorrectionColorSlider.tsx
@@ -0,0 +1,199 @@
+import React from 'react';
+import { SliderMode, ColorMode } from './ColorCorrection';
+import { SliderWheel, Slider, ValueInput } from '../../controls';
+import { ColorProperty } from './ColorCorrection';
+import { WidgetUtilities } from 'src/utilities';
+import { PropertyType } from 'src/shared';
+
+
+type Props = {
+  className?: string;
+  label: string;
+  mode: SliderMode;
+  sensitivity: number;
+  color: any;
+  colorProperty: ColorProperty;
+  property: string;
+  sliderProperty: string;
+  colorMode: ColorMode;
+  range: { min: number, max: number };
+  propertyType: PropertyType;
+
+  onPropertyChange: (property: string, v: any) => void;
+  onRangeUpdate: (color: any) => void;
+}
+
+type State = {
+
+}
+
+export class CorrectionColorSlider extends React.Component<Props, State> {
+
+  onChange = (key: string, sign: number, v?: number) => {
+    const { colorMode } = this.props;
+
+    if (colorMode === ColorMode.Rgb)
+      return this.onRgbChange(key, sign, v);
+
+    return this.onHsvChange(key, sign, v);
+  }
+
+  onRgbChange = (key: string, sign: number, v?: number) => {
+    const { sensitivity, color, colorProperty, property } = this.props;
+
+    if (!color)
+      return;
+
+    const step = sign * 0.01 * sensitivity;
+    color[key] += step;
+
+    if (v !== undefined)
+      color[key] = v;
+
+    if (colorProperty !== ColorProperty.Offset)
+      color[key] = Math.max(0, color[key]);
+
+    this.props.onPropertyChange(property, color);
+    this.props.onRangeUpdate(color);
+  }
+
+  onHsvChange = (key: string, sign: number, v?: number) => {
+    const { sensitivity, propertyType, range, colorProperty, property } = this.props;
+
+    if (!this.props.color)
+      return;
+
+    const step = sign * 0.01 * sensitivity;
+    const hsv = this.getHsv(range.max);
+
+    hsv[key] += step;
+
+    if (v !== undefined) {
+      if(key === 'v')
+        v = v / range.max;
+
+      hsv[key] = v;
+    }
+
+    hsv[key] = Math.max(hsv[key], 0.0001);
+    
+    const rgb = WidgetUtilities.hsv2rgb(hsv);
+    let color = rgb as any;
+
+    switch (propertyType) {
+      case PropertyType.Vector4:
+        color = WidgetUtilities.rgbToColor(rgb, range.max);
+        break;
+
+      case PropertyType.LinearColor:
+        color = WidgetUtilities.rgbToColor(rgb, range.max, false);
+        break;
+    }
+
+    if (rgb.A)
+      color.W = rgb.A;
+
+    if (colorProperty !== ColorProperty.Offset)
+      for (const key in color)
+        color[key] = Math.max(0, color[key]);
+
+    this.props.onPropertyChange(property, color);
+    this.props.onRangeUpdate(color);
+  }
+
+  getHsv = (max = 1) => {
+    const { color } = this.props;
+    const rgb = WidgetUtilities.colorToRgb(color, max);
+    
+    return WidgetUtilities.rgb2Hsv(rgb);
+  }
+
+  getValue = (key: string) => {
+    const { colorMode } = this.props;
+
+    if (colorMode === ColorMode.Rgb)
+      return this.getRgbValue(key);
+
+    return this.getHsvValue(key);
+  }
+
+  getRgbValue = (key: string) => {
+    const { color } = this.props;
+    const value = color?.[key] as number;
+
+    return value?.toFixed(3) ?? 0;
+  }
+
+  getHsvValue = (key: string) => {
+    const hsv = this.getHsv();
+    return hsv[key];
+  }
+  
+  getSSliderColor = () => {
+    const hsv = this.getHsv();
+    const rgb = WidgetUtilities.hsv2rgb({ h: hsv.h, s: 1, v: 1, a: hsv?.a });
+
+    return `linear-gradient(90deg, ${this.getVSliderMinColor()} 0%, rgba(${rgb.R}, ${rgb.G}, ${rgb.B}, ${rgb?.A ?? 1}) 100%)`;
+  }
+
+  getVSliderColor = () => {
+    const hsv = this.getHsv();
+    const rgb = WidgetUtilities.hsv2rgb({ h: hsv.h, s: hsv.s, v: 1, a: hsv?.a });
+
+    return `linear-gradient(90deg, rgba(0,0,0,${rgb?.A ?? 1}) 0%, rgba(${rgb.R}, ${rgb.G}, ${rgb.B}, ${rgb?.A ?? 1}) 100%)`;
+  }
+
+  getVSliderMinColor = () => {
+    const hsv = this.getHsv();
+    const rgb = WidgetUtilities.hsv2rgb({ h: hsv.h, s: 0, v: hsv.v, a: hsv?.a });
+
+    return `rgba(${rgb.R}, ${rgb.G}, ${rgb.B}, ${rgb?.A ?? 1})`;
+  }
+
+  render() {
+    let { className = '', mode, label, sliderProperty, range, colorMode } = this.props;
+    let sliderMax = Math.max(range.min, range.max);
+
+    const wheel = mode === SliderMode.Infinity;
+    const style: React.CSSProperties = {};
+
+    className += ` slider-wheel-container ${sliderProperty}-slider`;
+
+    if (!wheel)
+      className += ` slider`;
+
+    if (colorMode === ColorMode.Hsv && sliderProperty === 's') {
+      style.background = this.getSSliderColor();
+      sliderMax = 1;
+    };
+
+    if (colorMode === ColorMode.Hsv && sliderProperty === 'v') {
+      style.background = this.getVSliderColor();
+    };
+      
+    return (
+      <div className={className}>
+        <span className="title">{label}</span>
+        {wheel &&
+          <SliderWheel size={200}
+                       style={style}
+                       onWheelMove={sign => this.onChange(sliderProperty, sign)} />
+        }
+        {!wheel &&
+          <Slider showLabel={false}
+                  style={style}
+                  value={+this.getValue(sliderProperty)}
+                  onChange={value => this.onChange(sliderProperty, 0, value)}
+                  min={range.min}
+                  max={sliderMax}
+                  size={200} />
+        }
+        <ValueInput draggable={false}
+                    precision={3}
+                    min={range.min}
+                    value={this.getValue(sliderProperty)}
+                    onChange={value => this.onChange(sliderProperty, 0, value)} />
+      </div>
+    );
+  }
+};
\ No newline at end of file
diff --git a/Engine/Plugins/VirtualProduction/RemoteControlWebInterface/WebApp/Client/src/components/Screens/ColorCorrection/PaginationContent.tsx b/Engine/Plugins/VirtualProduction/RemoteControlWebInterface/WebApp/Client/src/components/Screens/ColorCorrection/PaginationContent.tsx
new file mode 100644
--- /dev/null
+++ b/Engine/Plugins/VirtualProduction/RemoteControlWebInterface/WebApp/Client/src/components/Screens/ColorCorrection/PaginationContent.tsx
@@ -0,0 +1,80 @@
+import React from 'react';
+
+
+type Props = {
+  pageItems?: number;
+  total?: number;
+  pages?: number;
+
+  children: ({ ...state }: State) => any;
+}
+
+type State = {
+  pagination: number;
+}
+
+export class PaginationContent extends React.Component<Props, State> {
+
+  state: State = {
+    pagination: 0,
+  }
+
+  getIsDisabled = (dir: number) => {
+    const { pagination } = this.state;
+    const { pageItems, total, pages } = this.props;
+
+    let step = pagination + dir;
+
+    if (step < 0)
+      return true;
+
+    if (pages && step > pages - 1)
+      return true;
+
+    return (step * pageItems) >= total;
+  }
+
+  onPaginationClick = (dir: number, disabled: boolean) => {
+    if (disabled)
+      return;
+
+    let { pagination } = this.state;
+    pagination += dir;
+
+    this.setState({ pagination });
+  }
+
+  renderPaginationArrow = (dir: number) => {
+    let className = 'arrow up ';
+    const disabled = this.getIsDisabled(dir);
+
+    if (dir > 0)
+      className = className.replace('up', 'down');
+
+    if (disabled)
+      className += 'disabled ';
+
+    return (
+      <div className={className} 
+           onClick={() => this.onPaginationClick(dir, disabled)}>
+        <div className='icon' />
+      </div>
+    );
+  }
+
+  render() {
+    const { children } = this.props;
+
+    return (
+      <div className="content" style={{ display: 'flex' }}>
+        <div className="body" style={{ width: '100%' }}>
+          {children({ ...this.state })}
+        </div>
+        <div className="arrows-pagination">
+          {this.renderPaginationArrow(-1)}
+          {this.renderPaginationArrow(1)}
+        </div>
+      </div>
+    );
+  }
+};
\ No newline at end of file
diff --git a/Engine/Plugins/VirtualProduction/RemoteControlWebInterface/WebApp/Client/src/components/Screens/ColorCorrection/index.ts b/Engine/Plugins/VirtualProduction/RemoteControlWebInterface/WebApp/Client/src/components/Screens/ColorCorrection/index.ts
new file mode 100644
--- /dev/null
+++ b/Engine/Plugins/VirtualProduction/RemoteControlWebInterface/WebApp/Client/src/components/Screens/ColorCorrection/index.ts
@@ -0,0 +1,4 @@
+export * from './ColorCorrection';
+export * from './CorrectionColorPicker';
+export * from './CorrectionColorSlider';
+export * from './PaginationContent';
\ No newline at end of file
diff --git a/Engine/Plugins/VirtualProduction/RemoteControlWebInterface/WebApp/Client/src/components/Screens/LightCards.tsx b/Engine/Plugins/VirtualProduction/RemoteControlWebInterface/WebApp/Client/src/components/Screens/LightCards.tsx
new file mode 100644
--- /dev/null
+++ b/Engine/Plugins/VirtualProduction/RemoteControlWebInterface/WebApp/Client/src/components/Screens/LightCards.tsx
@@ -0,0 +1,208 @@
+import React from 'react';
+import { _api } from 'src/reducers';
+import _ from 'lodash';
+import { ColorPicker, Slider } from '../controls';
+import { PropertyType } from 'src/shared';
+import { FontAwesomeIcon } from '@fortawesome/react-fontawesome';
+
+
+type Props = {
+};
+
+type State = {
+  selected: string;
+  actors: Record<string, string>;
+  value: any;
+};
+
+
+export class LightCards extends React.Component<Props, State> {
+
+  timer: NodeJS.Timeout;
+  state: State = {
+    selected: '',
+    actors: {},
+    value: {},
+  };
+
+  async componentDidMount() {
+    await this.onRefresh();
+    this.timer = setInterval(this.onRefresh, 1000);
+  }
+
+  componentWillUnmount() {
+    clearInterval(this.timer);
+    this.timer = null;
+  }
+
+  onRefresh = async () => {
+    await this.onRefreshActors();
+    await this.onRefreshValues();
+  }
+
+  onActorSelect = (selected: string) => {
+    this.setState({ selected }, this.onRefreshValues);
+  }
+
+  onSpawnActor = async () => {
+    await _api.proxy.function(
+      '/Script/RemoteControlWebInterface.Default__RCWebInterfaceBlueprintLibrary',
+      'SpawnActor',
+      {
+        Class: '/nDisplay/LightCard/LightCard.LightCard_C'
+      });
+
+    await this.onRefreshActors();
+  }
+
+  onRefreshActors = async () => {
+    const actors: Record<string, string> = await _api.proxy.function(
+      '/Script/RemoteControlWebInterface.Default__RCWebInterfaceBlueprintLibrary',
+      'FindAllActorsOfClass',
+      {
+        Class: '/nDisplay/LightCard/LightCard.LightCard_C'
+      });
+
+    if (!_.isEqual(this.state.actors, actors)) {
+      let { selected } = this.state;
+      
+      if (selected && !actors[selected])
+        selected = undefined;
+
+      if (!selected)
+        selected = _.first(Object.keys(actors));
+        
+      this.setState({ actors, selected });
+    }
+ 
+  }
+
+  onRefreshValues = async () => {
+    const { selected } = this.state;
+    if (!selected)
+      return;
+
+    const value = await _api.proxy.put('/remote/object/property', { objectPath: selected, access: 'READ_ACCESS' });
+    if (!_.isEqual(this.state.value, value)) {
+      this.setState({ value });
+    }
+  }
+
+  onPropertyValueChangeInternal = async (propertyName: string, property: string, value: any) => {
+    const { selected } = this.state;
+    if (!selected)
+      return;
+    
+    const propertyValue = _.set({}, property, value);
+    await _api.proxy.put('/remote/object/property', { 
+      objectPath: selected,
+      access: 'WRITE_ACCESS',
+      propertyName,
+      propertyValue
+    });
+  }
+
+  onThrottledPropertyValueChange = _.throttle(this.onPropertyValueChangeInternal, 100);
+
+  onPropertyChange = (property: string, v: any) => {
+    const { value } = this.state;
+    _.set(value, property, v);
+    this.setState({ value });
+
+    const [ propertyName ] = property.split('.');
+    this.onThrottledPropertyValueChange(propertyName, property, v);
+  }
+
+  renderColorPicker = (property: string) => {
+    const { value } = this.state;
+
+    return (
+      <ColorPicker value={_.get(value, property)}
+                    type={PropertyType.LinearColor}
+                    alpha={true}
+                    onChange={value => this.onPropertyChange(property, value)} />
+    );
+  }
+
+  renderSlider = (property: string, min: number, max: number, label?: string) => {
+    const { value } = this.state;
+
+    return (
+      <div className='control'>
+        <label>{label ?? _.startCase(property)}</label>
+        <Slider showLabel={true}
+                value={_.get(value, property)}
+                onChange={value => this.onPropertyChange(property, value)}
+                min={min}
+                max={max} />
+    </div>
+    );
+  }
+
+  render() {
+    const { selected, actors } = this.state;
+
+    const style: React.CSSProperties = {};
+    if (!selected) {
+      style.opacity = 0.6;
+      style.pointerEvents = 'none';
+    }
+
+    return (
+      <div className="light-cards-wrapper  custom-ui-wrapper" tabIndex={-1}>
+        <div className='group' style={style}>
+
+          <div className='group vertical'>
+            <div className='cc-group'>
+              <div className='cc-title'>Size and Position</div>
+              <div className='body'>
+                {this.renderSlider('DistanceFromCenter', 0, 10_000, 'Distance')}
+                {this.renderSlider('Scale.X', 0.01, 1)}
+                {this.renderSlider('Scale.Y', 0.01, 1)}
+                {this.renderSlider('Longitude', 0, 360)}
+                {this.renderSlider('Latitude', 0, 360)}
+              </div>
+            </div>
+
+            <div className='cc-group'>
+              <div className='cc-title'>Guides</div>
+              <div className='body'>
+                {this.renderColorPicker('Color')}
+              </div>
+            </div>
+          </div>
+
+
+          <div className='cc-group'>
+            <div className='cc-title'>Appearance</div>
+            <div className='body'>
+              {this.renderColorPicker('GuideColor')}
+
+              {this.renderSlider('Exposure', 0, 10)}
+              {this.renderSlider('Opacity', 0, 1)}
+              {this.renderSlider('Feathering', 0, 10)}
+            </div>
+          </div>
+
+        </div>
+
+        <div className='cc-group shotbox'>
+          <div className='body'>
+            {Object.keys(actors).map(actor => 
+              <div key={actor}
+                  className={`light-card ${selected === actor ? 'active' : ''}`}
+                  onClick={() => this.onActorSelect(actor)}>
+                {_.startCase(actors[actor])}
+              </div>
+            )}
+
+            <div className="light-card active"
+                onClick={() => this.onSpawnActor()}>
+              <FontAwesomeIcon icon={['fas', 'plus']} />
+            </div>                
+          </div>
+        </div>
+      </div>
+    );
+  }
+}
diff --git a/Engine/Plugins/VirtualProduction/RemoteControlWebInterface/WebApp/Client/src/components/Screens/Playlist.tsx b/Engine/Plugins/VirtualProduction/RemoteControlWebInterface/WebApp/Client/src/components/Screens/Playlist.tsx
new file mode 100644
--- /dev/null
+++ b/Engine/Plugins/VirtualProduction/RemoteControlWebInterface/WebApp/Client/src/components/Screens/Playlist.tsx
@@ -0,0 +1,345 @@
+import React from 'react';
+import { FontAwesomeIcon } from '@fortawesome/react-fontawesome';
+import { _api } from 'src/reducers';
+import _ from  'lodash';
+import { IAsset } from 'src/shared';
+import { SafeImage, ValueInput } from '../controls';
+import { SearchModal } from '../';
+import path from 'path';
+
+
+interface IPlaylistData {
+  Description: string;
+  Items: string[];
+  PlaylistItems: IPlaylistItem[];
+}
+
+interface IPlaylistItem {
+  Path: string;
+  IsPlaying?: boolean;
+
+  Sequence: string;
+	StartFrameOffset: number;
+	EndFrameOffset: number;
+	bHoldAtFirstFrame: boolean;
+	NumLoops: number;
+  PlaybackSpeed: number;
+	bMute: boolean;
+}
+
+type Props = {
+  visible?: boolean;
+};
+
+type State = {
+  filter: string;
+  sequences: IAsset[];
+  playlist?: IAsset;
+  data?: IPlaylistData;
+  activeTooltip: number;
+};
+
+export class Playlist extends React.Component<Props, State> {
+  player: string;
+  interval: NodeJS.Timer;
+
+  state: State = {
+    filter: '',
+    sequences: [],
+    activeTooltip: null,
+  };
+
+  componentDidMount() {
+    document.onclick = this.onDocumentClick;
+    this.onInitialize();
+    this.interval = setInterval(this.onCheckSequenceStatus, 300);
+  }
+
+  componentWillUnmount() {
+    document.onclick = null;
+    clearInterval(this.interval);
+    this.interval = null;
+  }
+
+  onInitialize = async () => {
+    const WorldContextObject = await _api.proxy.function("Object'/Engine/Transient.UnrealEdEngine_0:UnrealEditorSubsystem_0'", 'GetEditorWorld');
+    if (!WorldContextObject)
+      return;
+
+    this.player = await _api.proxy.function("Object'/Script/SequencerPlaylists.Default__SequencerPlaylistPlayer'", 'GetDefaultPlayer', { WorldContextObject });
+
+    const sequences = await _api.assets.search('', ['LevelSequence'], '/Game', 200);
+    this.setState({ sequences });
+
+    const currentPlaylist = await _api.proxy.function(this.player, 'GetPlaylist');
+    if (currentPlaylist) {
+      const folder = path.dirname(currentPlaylist);
+      const [filename] = path.basename(currentPlaylist).split('.');
+
+      const playlistAsset = await _api.assets.search(filename, ['SequencerPlaylist'], folder, 200);
+      await this.onLoadPlaylist(playlistAsset[0]);
+    }
+  }
+
+  onSearchChange = (filter: string)  => {
+    this.setState({ filter });
+  }
+
+  onCheckSequenceStatus = async () => {
+    const { playlist, data } = this.state;
+    if (!playlist || !data?.Items)
+      return;
+
+    let changed = false;
+    for (const item of data.PlaylistItems) {
+      const playing = !!await _api.proxy.function(this.player, 'IsPlaying', { Item: item.Path });
+
+      if (playing !== item.IsPlaying) {
+        item.IsPlaying = playing;
+        changed = true;
+      }
+    }
+
+    if (changed)
+      this.forceUpdate();
+  }
+
+  onDocumentClick = (e: MouseEvent) => {
+    const { activeTooltip } = this.state;
+
+    if (e.defaultPrevented || activeTooltip === null)
+      return;
+
+    this.setState({ activeTooltip: null });
+  }
+
+  onTooltipClick = (e: React.MouseEvent, activeTooltip: number) => {
+    e.preventDefault();
+    this.setState({ activeTooltip });
+  }
+
+  onPlaylistBrowse = async () => {
+    const playlist = await SearchModal.open({ placeholder: 'Search Playlist', types: ['SequencerPlaylist'], prefix: '/Game', count: 200 });
+    if (!playlist)
+      return;
+
+    await _api.proxy.function(this.player, 'SetPlaylist', { InPlaylist: playlist.Path });
+    await this.onLoadPlaylist(playlist);
+  }
+
+  onLoadPlaylist = async (playlist: IAsset) => {
+    if (!playlist)
+      return;
+
+    const data = await _api.proxy.property.get<IPlaylistData>(playlist.Path);
+    data.PlaylistItems = [];
+    for (const itemPath of data.Items) {
+      const item = await _api.proxy.property.get<IPlaylistItem>(itemPath);
+      item.Path = itemPath;
+      data.PlaylistItems.push(item);
+    }
+
+    this.setState({ playlist, data, filter: '' });
+  }
+
+  onAddItem = async () => {
+    // const { playlist, data } = this.state;
+    // if (!playlist || !data?.Items)
+    //   return;
+
+    // data.Items.push('');
+    // await _api.proxy.property.set(playlist.Path, 'Items', data.Items);
+  }
+
+  onSequenceBrowse = async (item: IPlaylistItem) => {
+    const sequence = await SearchModal.open({ placeholder: 'Search Sequence', types: ['LevelSequence'], prefix: '/Game', count: 200 });    
+    if (sequence)
+      await this.onSetPlaylistItem(item, 'Sequence', sequence.Path);
+  }
+
+  onPlaylistAction = (action: string) => {
+    const { playlist } = this.state;
+    if (!playlist || !this.player)
+      return;
+
+    return _api.proxy.function(this.player, action);
+  }
+
+  onPlaylistItemAction = (item: IPlaylistItem, action: string) => {
+    const { playlist } = this.state;
+    if (!playlist || !this.player)
+      return;
+
+    return _api.proxy.function(this.player, action, { Item: item.Path });
+  }
+
+
+  onSetPlaylistItem = async (item: IPlaylistItem, property: string, value: any) => {
+    item[property] = value;
+    await _api.proxy.property.set(item.Path, property, value);
+    this.forceUpdate();
+  }
+
+  renderPlaylistItem = (item: IPlaylistItem, index: number) => {
+    const { sequences } = this.state;
+
+    const url = _api.assets.thumbnailUrl(item.Sequence);
+
+    let playButtonClass: string = "";
+    if (item.IsPlaying)
+      playButtonClass = "active";
+
+    return (
+      <tr key={index} className={`playlist-item col col-${index + 1}`}>
+        <td className="col-actions">
+          <div className="controls">
+            <button tabIndex={-1} className={`btn btn-secondary playlist-play ${(item.IsPlaying) && 'playing'} `} onClick={() => this.onPlaylistItemAction(item, 'PlayItem')}>
+              <FontAwesomeIcon icon={['fas', 'play']} />
+            </button>
+            <button tabIndex={-1} className={`btn btn-secondary playlist-stop  ${(item.IsPlaying) && 'playing active'}`} onClick={() => this.onPlaylistItemAction(item, 'StopItem')}>
+              <FontAwesomeIcon icon={['fas', 'stop']} />
+            </button>
+            <button tabIndex={-1} className="btn btn-secondary playlist-rehold" onClick={() => this.onPlaylistItemAction(item, 'ResetItem')}>
+              <FontAwesomeIcon icon={['fas', 'reply']} />
+            </button>
+          </div>
+        </td>
+
+        <td className='col-thumb'>
+          <div className="thumbnail-block">
+            <SafeImage src={url} 
+                        fallback="/images/assets/Sequencer.png"
+                        alt={item.Sequence}
+                        className="box-image" />
+          </div>
+        </td>
+        <td className='col-name'>
+          <select value={item.Sequence} onChange={e => this.onSetPlaylistItem(item, 'Sequence', e.target.value)}>
+            <option>None</option>
+            {sequences.map(seq => <option key={seq.Path} value={seq.Path}>{seq.Name}</option>)}
+          </select>
+          <span className="icon" onClick={() => this.onSequenceBrowse(item)}>
+          <img src='/images/icons/playlistIcons/ContentBrowser.svg' className='playlist-icon'/>
+          </span>
+        </td>
+        <td className='col-control'>
+          <div className='input-container'>
+            <ValueInput value={item.StartFrameOffset} onChange={value => this.onSetPlaylistItem(item, 'StartFrameOffset', value)} precision={0} />
+          </div>
+          <div className='input-container'>
+            <ValueInput value={item.EndFrameOffset} onChange={value => this.onSetPlaylistItem(item, 'EndFrameOffset', value)} precision={0} />
+          </div>
+        </td>
+        <td className='col-control'>
+          
+          <div className='checkbox-container'>
+            <label className="container">
+              <input type="checkbox" checked={item.bHoldAtFirstFrame} onChange={e => this.onSetPlaylistItem(item, 'bHoldAtFirstFrame', e.target.checked)}/>
+              <span className="checkmark">
+                <img src='/images/icons/playlistIcons/HoldFrame.svg' className='playlist-icon'/>
+              </span>
+            </label>
+          </div>
+        </td>
+        <td className='col-control'>
+          <div>
+            <ValueInput value={item.NumLoops} onChange={value => this.onSetPlaylistItem(item, 'NumLoops', value)} precision={0} />
+          </div>
+        </td>
+      </tr>
+    );
+  }
+
+  render() {
+    const { filter, playlist, data } = this.state;
+
+    let items = data?.PlaylistItems ?? [];
+    if (filter) {
+      const filterLowerCase = filter.toLowerCase();
+      items = _.filter(items, i => i.Sequence.toLowerCase().includes(filterLowerCase));
+    }
+
+    if (!playlist)
+      return (
+        <div className="playlist-wrapper">
+          <div className="playlist-placeholder">
+            <div className="playlist-placeholder-title" onClick={this.onPlaylistBrowse}>
+              <img src='/images/icons/playlistIcons/PlaylistLoad.svg' className='playlist-icon large'/>
+              Load Playlist
+            </div>
+          </div>
+        </div>
+      );
+
+    return (
+      <div className="playlist-wrapper" tabIndex={-1}>
+        <div className="playlist-controls">
+          <div>
+            <div className="playlist-header" onClick={this.onPlaylistBrowse}>
+              <span className="icon"><FontAwesomeIcon icon={['fas', 'folder']} /></span>
+              <h1>{playlist?.Name}</h1>
+            </div>
+            <span className="description">{data?.Description}</span>
+          </div>
+          <div className="controls">
+            <button tabIndex={-1} className="btn btn-secondary playlist-play" onClick={() => this.onPlaylistAction('PlayAll')}>
+              <FontAwesomeIcon icon={['fas', 'play']} />
+            </button>
+            <button tabIndex={-1} className="btn btn-secondary playlist-stop" onClick={() => this.onPlaylistAction('StopAll')}>
+              <FontAwesomeIcon icon={['fas', 'stop']} />
+            </button>
+            <button tabIndex={-1} className="btn btn-secondary playlist-rehold" onClick={() => this.onPlaylistAction('ResetAll')}>
+              <FontAwesomeIcon icon={['fas', 'reply']} />
+            </button>
+          </div>
+        </div>
+        <div className="list-wrapper">
+          <div className="search-block">
+          <button tabIndex={-1} className="btn btn-secondary add-item" onClick={this.onAddItem}><FontAwesomeIcon icon={['fas', 'plus']} /> Item</button>
+            <div className="search-field">
+              <FontAwesomeIcon icon={['fas', 'search']} />
+              <input type="text" value={filter} onChange={e => this.onSearchChange(e.target.value)} autoComplete="no" />
+            </div>
+          </div>
+
+          <div className="table-wrapper">
+            <table>
+              <thead>
+                <tr>
+                  <th className="col col-control">&nbsp;</th>
+                  <th className="col col-3" colSpan={2}>Playlist Items</th>
+                  <th className="col col-control">Offset</th>
+                  <th className="col col-control">Hold</th>
+                  <th className="col col-control">Loop</th>
+                </tr>
+              </thead>
+              <tbody>
+                {items.map(this.renderPlaylistItem)}
+              </tbody>
+            </table>
+          </div>
+        </div>
+      </div>
+    );
+  }
+}
+
+type PlaylistTooltipProps = {
+  visible: boolean;
+}
+
+
+class PlaylistTooltip extends React.Component<PlaylistTooltipProps> {
+  render() {
+    const { visible } = this.props;
+
+    if (!visible)
+      return null;
+
+    return (
+      <div className="playlist-tooltip">
+        <div className="playlist-tooltip-header">Item Details</div>
+      </div>
+    );
+  }
+}
diff --git a/Engine/Plugins/VirtualProduction/RemoteControlWebInterface/WebApp/Client/src/components/Screens/Snapshot.tsx b/Engine/Plugins/VirtualProduction/RemoteControlWebInterface/WebApp/Client/src/components/Screens/Snapshot.tsx
--- a/Engine/Plugins/VirtualProduction/RemoteControlWebInterface/WebApp/Client/src/components/Screens/Snapshot.tsx
+++ b/Engine/Plugins/VirtualProduction/RemoteControlWebInterface/WebApp/Client/src/components/Screens/Snapshot.tsx
@@ -45,15 +45,14 @@ export class Snapshot extends React.Component<Props, State> {
 
     if (!map) {
 	    const editorSubsystem = await _api.proxy.function('/Script/UnrealEd.Default__EditorSubsystemBlueprintLibrary', 'GetEditorSubsystem', { Class: '/Script/UnrealEd.UnrealEditorSubsystem' });
-      const world = await _api.proxy.function(editorSubsystem.ReturnValue, 'GetEditorWorld', {});
-      map = world.ReturnValue;
+      map = await _api.proxy.function(editorSubsystem, 'GetEditorWorld', {});
     }
 
     const nativeFilters = await _api.assets.search('', ['LevelSnapshotFilter'], '/Game', {}, 50);
 
     const filterArgs = {
       EnableBlueprintNativeClassFiltering: true,
-      NativeParentClasses: ['LevelSnapshotFilter'],        
+      NativeParentClasses: ['LevelSnapshotFilter'],
     };
     const bpFilters = _.filter(await _api.assets.search('', ['Blueprint'], '/Game', filterArgs, 50), f => !!f.Metadata['GeneratedClass']);
 
@@ -218,18 +217,17 @@ class SnapshotPreview extends React.Component<SnapshotPreviewProps, SnapshotPrev
 
         switch (filterAsset?.Class) {
           case 'Blueprint':
-            const createdFilter = await _api.proxy.function(
+            filterInstance  = await _api.proxy.function(
               '/Script/LevelSnapshotFilters.Default__FilterBlueprintFunctionLibrary', 
               'CreateFilterByClass',
               { 
                 Class: filterAsset.Metadata['GeneratedClass']
               }
             );
-            filterInstance = createdFilter?.ReturnValue;
             break;
 
-	  default: //'LevelSnapshotFilter'
-	    filterInstance = filterAsset.Path;
+          case 'LevelSnapshotFilter':
+            filterInstance = filterAsset.Path;
             break;
         }
       }
@@ -358,7 +356,7 @@ class SaveModal extends React.Component<SaveModalProps, SaveModalState> {
       const PathFormatted = await _api.proxy.function(editor, 'ParseLevelSnapshotsTokensInText', { InTextToParse: path, InWorldName });
       const NameFormatted = await _api.proxy.function(editor, 'ParseLevelSnapshotsTokensInText', { InTextToParse: DefaultLevelSnapshotName, InWorldName });
 
-      this.setState({ path: PathFormatted.ReturnValue, name: NameFormatted.ReturnValue });
+      this.setState({ path: PathFormatted, name: NameFormatted });
     } catch {
 
     }
@@ -445,4 +443,4 @@ class SaveModal extends React.Component<SaveModalProps, SaveModalState> {
       </div>
     );
   }
-}
+}
\ No newline at end of file
diff --git a/Engine/Plugins/VirtualProduction/RemoteControlWebInterface/WebApp/Client/src/components/Screens/index.ts b/Engine/Plugins/VirtualProduction/RemoteControlWebInterface/WebApp/Client/src/components/Screens/index.ts
--- a/Engine/Plugins/VirtualProduction/RemoteControlWebInterface/WebApp/Client/src/components/Screens/index.ts
+++ b/Engine/Plugins/VirtualProduction/RemoteControlWebInterface/WebApp/Client/src/components/Screens/index.ts
@@ -1,2 +1,5 @@
+export * from './ColorCorrection';
+export * from './LightCards';
+export * from './Playlist';
 export * from './Snapshot';
 export * from './Sequencer';
diff --git a/Engine/Plugins/VirtualProduction/RemoteControlWebInterface/WebApp/Client/src/components/SignalIcon.tsx b/Engine/Plugins/VirtualProduction/RemoteControlWebInterface/WebApp/Client/src/components/SignalIcon.tsx
new file mode 100644
--- /dev/null
+++ b/Engine/Plugins/VirtualProduction/RemoteControlWebInterface/WebApp/Client/src/components/SignalIcon.tsx
@@ -0,0 +1,24 @@
+import React from 'react';
+import { ConnectionSignal } from 'src/shared';
+
+
+type Props = {
+  signal: ConnectionSignal;
+}
+
+export class SignalIcon extends React.Component<Props> {
+  render() {
+    const { signal } = this.props;
+    if (signal === ConnectionSignal.Good)
+      return null;
+
+    return (
+      <div className={`signal-icon ${signal.toLowerCase()}`}>
+        <span className="bar-1"></span>
+        <span className="bar-2"></span>
+        <span className="bar-3"></span>
+        <span className="bar-4"></span>
+      </div>
+    );
+  }
+};
\ No newline at end of file
diff --git a/Engine/Plugins/VirtualProduction/RemoteControlWebInterface/WebApp/Client/src/components/Stack.tsx b/Engine/Plugins/VirtualProduction/RemoteControlWebInterface/WebApp/Client/src/components/Stack.tsx
--- a/Engine/Plugins/VirtualProduction/RemoteControlWebInterface/WebApp/Client/src/components/Stack.tsx
+++ b/Engine/Plugins/VirtualProduction/RemoteControlWebInterface/WebApp/Client/src/components/Stack.tsx
@@ -2,7 +2,7 @@ import React from 'react';
 import { connect } from 'react-redux';
 import { _api, ReduxState } from 'src/reducers';
 import { ColorProperty, ICustomStackProperty, ICustomStackTabs, ICustomStackWidget, IPanel, IPanelType, IPreset,
-        IPayload, PropertyType, PropertyValue, VectorProperty, WidgetTypes, IExposedProperty, IExposedFunction, IColorPickerList, ICustomStackListItem, IProperty } from 'src/shared';
+        IPayload, PropertyType, PropertyValue, VectorProperty, WidgetTypes, IExposedProperty, IExposedFunction, IColorPickerList, ICustomStackListItem } from 'src/shared';
 import { ColorPicker, ColorPickerList, Button, Text, VectorDrawer, Dial, JoysticksWrapper, DialMode, DialsWrapper } from './controls';
 import { PrecisionModal } from '../components';
 import { FontAwesomeIcon } from '@fortawesome/react-fontawesome';
@@ -214,8 +214,11 @@ export class Stack extends React.Component<Props & PropsFromState, State> {
     const type = property?.Type;
     const precision = WidgetUtilities.getPropertyPrecision(type);
     const alpha = !!this.getMetadata(property, 'Alpha');
+    const proportionally = !!this.getMetadata(property, 'Proportionally');
+
     let min = this.getMetadata(property, 'Min');
     let max = this.getMetadata(property, 'Max');
+
     if (type === PropertyType.Rotator) {
       if (min === undefined)
         min = -180;
@@ -223,7 +226,7 @@ export class Stack extends React.Component<Props & PropsFromState, State> {
         max = 180;
     }
 
-    return { type, min, max, precision, alpha };
+    return { type, min, max, precision, alpha, proportionally };
   }
 
   getAlpha = (property: string) => {
@@ -269,7 +272,7 @@ export class Stack extends React.Component<Props & PropsFromState, State> {
                       type={type}
                       value={value}
                       onChange={onChange}
-                      onMetadataChange={(key, value) => _api.payload.metadata(widget.property, key, value)}
+                      onMetadataChange={(key, value) => this.onMetadataChange(widget.property, key, value)}
                       onClose={() => this.setState({ precisionModal: null })} />
     );
   }
@@ -321,18 +324,20 @@ export class Stack extends React.Component<Props & PropsFromState, State> {
                        selected={this.props.selected === selection}
                        isDragDisabled={!this.props.editable}>
         <Widgets.SlidersWidget widget={widget}
-                                label={this.getLabel(widget)}
-                                min={meta.min}
-                                max={meta.max}
-                                value={this.getValue(widget)}
-                                onChange={this.onPropertyAxisValueChange}
-                                onPrecisionModal={property => this.setState({ precisionModal: { widget, property } })} />
+                               label={this.getLabel(widget)}
+                               min={meta.min}
+                               max={meta.max}
+                               proportionally={meta.proportionally}
+                               value={this.getValue(widget)}
+                               onChange={this.onPropertyAxisValueChange}
+                               onPrecisionModal={property => this.setState({ precisionModal: { widget, property } })}
+                               onProportionallyToggle={this.onProportionallyToggle} />
       </DraggableWidget>
     );
   }
 
   renderColorPickerList = (widget: IColorPickerList, index: number, path: string) => {
-    const { droppable, selected, editable, dragging, preset } = this.props;
+    const { droppable, selected, editable, dragging } = this.props;
 
     const selection = `${path}_${index}_${WidgetTypes.ColorPickerList}`;
     const draggableId = `${widget.id}_${WidgetTypes.ColorPickerList}` || tempId;
@@ -740,12 +745,37 @@ export class Stack extends React.Component<Props & PropsFromState, State> {
     );
   }
 
-  renderDialsWidget = (widget: ICustomStackProperty, index: number, path: string) => {    
+  renderAssetWidget = (widget: ICustomStackProperty, index: number, path: string) => {
+    const { preset } = this.props;
+    const id = this.getWidgetId(widget, index);
+    const selection = `${path}_${index}_${widget.property}`;
+
+    const exposed = preset?.Exposed[widget.property] as IExposedProperty;
+
+    return (
+      <DraggableWidget className="asset-row"
+                       key={id}
+                       index={index}
+                       draggableId={id}
+                       onPointerChange={this.onDragChange}
+                       onSelect={this.props.onSelected.bind(this, selection)}
+                       selected={this.props.selected === selection}
+                       isDragDisabled={!this.props.editable}>
+        <Widgets.AssetWidget label={this.getLabel(widget)}
+                             value={this.getValue(widget)}
+                             type={exposed?.Type ?? 'Object'}
+                             onChange={value => this.onPropertyValueChange(widget, value)}
+                             browse />
+      </DraggableWidget>
+    );
+  }
+
+  renderDialsWidget = (widget: ICustomStackProperty, index: number, path: string) => {
     const id = this.getWidgetId(widget, index);    
     const selection = `${path}_${index}_${widget.property}`;
 
     const rotator = widget.propertyType === PropertyType.Rotator;
-    const { min, max } = this.getWidgetMetadata(widget);
+    const { min, max, proportionally } = this.getWidgetMetadata(widget);
 
     let dialMode = DialMode.Endless;
 
@@ -767,12 +797,15 @@ export class Stack extends React.Component<Props & PropsFromState, State> {
         <DialsWrapper type={widget.propertyType}
                       min={min}
                       max={max}
+                      proportionally={proportionally}
+                      property={widget.property}
                       mode={dialMode}
                       hidePrecision={true}
                       properties={widget?.widgets}
                       value={this.getValue(widget)}
                       label={this.getLabel(widget)}
-                      onChange={value => this.onThrottledPropertyValueChange(widget, value)} />
+                      onChange={value => this.onThrottledPropertyValueChange(widget, value)}
+                      onProportionallyToggle={this.onProportionallyToggle} />
       </DraggableWidget>
     );
   }
@@ -832,8 +865,8 @@ export class Stack extends React.Component<Props & PropsFromState, State> {
     );
   }
 
-  onPropertyAxisValueChange = (widget: ICustomStackProperty, axis: string, axisValue: number, locked?: boolean) => {
-    const value = this.getValue(widget) ?? {};
+  onPropertyAxisValueChange = (widget: ICustomStackProperty, axis: string, axisValue: number, proportionally?: boolean, min?: number, max?: number) => {
+    const value = { ...this.getValue(widget) } ?? {};
 
     // Reset
     if (axis === undefined || axisValue === undefined)
@@ -843,15 +876,18 @@ export class Stack extends React.Component<Props & PropsFromState, State> {
     if (prev === 0 || prev === undefined)
       prev = 1;
 
-    let ratio = Math.max(0.001, axisValue) / prev;
+    const ratio = Math.max(0.001, axisValue) / prev;
 
-    if (locked && !isNaN(ratio)) {
-      const keys = Object.keys(value);
-      for (const key of keys)
-        value[key] *= ratio;
-    } else {
+    if (proportionally && !isNaN(ratio))
+      for (const key of Object.keys(value)) {
+        let val = value[key] * ratio;
+        if (!isNaN(min) && !isNaN(max))
+          val = Math.min(max, Math.max(min, val));
+
+        value[key] = val;
+      }
+    else
       value[axis] = axisValue;
-    }
 
     this.onThrottledPropertyValueChange(widget, value);
   }
@@ -863,6 +899,14 @@ export class Stack extends React.Component<Props & PropsFromState, State> {
     this.props.onSetVectorDrawer(vector);
   }
 
+  onMetadataChange = (property: string, key: string, value: string) => {
+    _api.payload.metadata(property, key, value);
+  }
+
+  onProportionallyToggle = (property: string, value: string) => {
+    this.onMetadataChange(property, 'Proportionally', value);
+  }
+
   renderVectorWidget = (widget: ICustomStackProperty, index: number, path: string) => {
     const { vector } = this.props;
     const id = this.getWidgetId(widget, index);
@@ -884,9 +928,11 @@ export class Stack extends React.Component<Props & PropsFromState, State> {
                               vector={vector}
                               min={meta.min}
                               max={meta.max}
+                              proportionally={meta.proportionally}
                               onAxisChange={this.onPropertyAxisValueChange}
                               onToggleVectorDrawer={this.onToggleVectorDrawer}
-                              onSetVector={this.props.onSetVectorDrawer} />
+                              onSetVector={this.props.onSetVectorDrawer}
+                              onProportionallyToggle={this.onProportionallyToggle} />
       </DraggableWidget>
     );
   }
@@ -999,6 +1045,7 @@ export class Stack extends React.Component<Props & PropsFromState, State> {
       widget.widget = widgetType;
 
       let element = null;
+
       switch (widgetType) {
         case WidgetTypes.Tabs:
           element = this.renderTabs(widgets[index] as ICustomStackTabs, index, path);
@@ -1044,6 +1091,10 @@ export class Stack extends React.Component<Props & PropsFromState, State> {
           element = this.renderDialsWidget(widget, index, path);
           break;
 
+        case WidgetTypes.Asset:
+          element = this.renderAssetWidget(widget, index, path);
+          break;
+
         case WidgetTypes.Dial:
           firstIndex = index;
           w = this.getInlineWidgets(elements, widgets, WidgetTypes.Dial, index, 3);
@@ -1160,7 +1211,7 @@ export class Stack extends React.Component<Props & PropsFromState, State> {
 
 
     const templates = { ...this.state.templates };
-    templates[panel.id] = { matching: matching?.ReturnValue, selected: selected?.ReturnValue };
+    templates[panel.id] = { matching, selected };
 
     this.setState({ templates });
     sessionStorage.setItem('templates', JSON.stringify(templates));
@@ -1229,7 +1280,8 @@ export class Stack extends React.Component<Props & PropsFromState, State> {
                              isDropDisabled={isDropDisabled}>
               {!!panel?.title && <div className="title">{panel.title}</div>}
               {!!panel.isTemplate && (
-                <div className="title">
+                <div className="title template-property-select-wrapper">
+                  <div className='template-property-select-inner-wrapper'>
                     <select className="dropdown template-property-select" 
                             tabIndex={-1}
                             onFocus={this.isIos ? null : () => this.onRefreshPanelOptions(panel)}
@@ -1239,6 +1291,7 @@ export class Stack extends React.Component<Props & PropsFromState, State> {
                         <option key={match} value={match}>{match}</option>
                       )}
                     </select>
+                  </div>
                     {this.isIos && 
                       <button onClick={() => this.onRefreshPanelOptions(panel)}><FontAwesomeIcon icon={['fas', 'sync-alt']} /></button>
                     }
diff --git a/Engine/Plugins/VirtualProduction/RemoteControlWebInterface/WebApp/Client/src/components/Tabs.tsx b/Engine/Plugins/VirtualProduction/RemoteControlWebInterface/WebApp/Client/src/components/Tabs.tsx
--- a/Engine/Plugins/VirtualProduction/RemoteControlWebInterface/WebApp/Client/src/components/Tabs.tsx
+++ b/Engine/Plugins/VirtualProduction/RemoteControlWebInterface/WebApp/Client/src/components/Tabs.tsx
@@ -1,5 +1,6 @@
 import React from 'react';
 import { ITab, IPreset } from '../shared';
+import { _api } from '../reducers';
 import { FontAwesomeIcon } from '@fortawesome/react-fontawesome';
 import { IconName } from '@fortawesome/fontawesome-svg-core';
 import { Droppable, Draggable, DraggableStateSnapshot } from 'react-beautiful-dnd';
@@ -14,7 +15,11 @@ type Props = {
   tabs: ITab[];
   selected: number;
   editable?: boolean;
+  locked: boolean;
+  undo: boolean;
+  redo: boolean;
 
+  onLockedChange: () => void;
   onChange: (index: number) => void;
   onNewTab: () => void;
   onDeleteTab: (tab: number) => Promise<void>;
@@ -124,7 +129,7 @@ export class Tabs extends React.Component<Props> {
   }
 
   render() {
-    const { tabs, selected, editable, preset } = this.props;
+    const { tabs, selected, editable, preset, locked, undo, redo } = this.props;
     const { allTabsVisible } = this.state;
 
     let leftTabClassName = 'arrow left ';
@@ -135,6 +140,14 @@ export class Tabs extends React.Component<Props> {
     if (selected === tabs?.length - 1 || allTabsVisible)
       rightTabClassName += 'disabled ';
 
+    let iconUndoClassName = 'icon ';
+    if (!undo)
+      iconUndoClassName += 'disabled';
+      
+    let iconRedoClassName = 'icon ';
+    if (!redo)
+      iconRedoClassName += 'disabled';
+
     return (
       <div className="tabs-wrapper">
         <div className="tabs-wrapper-inner">
@@ -158,8 +171,19 @@ export class Tabs extends React.Component<Props> {
           }
           <div className={leftTabClassName} onClick={() => this.onTabChange(selected - 1)}><FontAwesomeIcon icon={['fas', 'caret-left']} /></div>
           <div className={rightTabClassName} onClick={() => this.onTabChange(selected + 1)}><FontAwesomeIcon icon={['fas', 'caret-right']} /></div>
+          <div className="preset-undo-redo">
+            <span className={iconUndoClassName} onClick={_api.undoHistory}><FontAwesomeIcon icon={['fas', 'undo']} /></span>
+            <span className={iconRedoClassName} onClick={_api.redoHistory}><FontAwesomeIcon icon={['fas', 'redo']} /></span>
+          </div>
           <div className="preset-name">{preset?.Name}</div>
+          <div className="tab edit-tab locked-toggle">
+            <label className="switch toggle-mode locked-mode">
+              <input type="checkbox" checked={!!locked} onChange={() => this.props.onLockedChange()} />
+              <span className="slider inline"></span>
+            </label>
+          </div>
         </div>
+        
       </div>
     );
   }
diff --git a/Engine/Plugins/VirtualProduction/RemoteControlWebInterface/WebApp/Client/src/components/Widgets/AssetWidget.tsx b/Engine/Plugins/VirtualProduction/RemoteControlWebInterface/WebApp/Client/src/components/Widgets/AssetWidget.tsx
new file mode 100644
--- /dev/null
+++ b/Engine/Plugins/VirtualProduction/RemoteControlWebInterface/WebApp/Client/src/components/Widgets/AssetWidget.tsx
@@ -0,0 +1,62 @@
+import { FontAwesomeIcon } from '@fortawesome/react-fontawesome';
+import React from 'react';
+import { _api } from '../../reducers';
+import { SafeImage } from '../controls';
+import { SearchModal } from '../';
+
+
+type Props = {
+  label?: string | JSX.Element;
+  value?: string;
+  type?: string;
+  browse?: boolean;
+  reset?: boolean;
+  onChange?: (value?: string) => void;
+}
+
+export class AssetWidget extends React.Component<Props> {
+
+  static defaultProps: Props = {
+    reset: true,
+  };
+
+  onBrowse = async () => {
+    let { type } = this.props;
+    if (type.startsWith('U') || type.startsWith('F'))
+      type = type.substring(1);
+
+    type = type.replace('*', '');
+
+    const asset = await SearchModal.open({ placeholder: 'Search asset', types: [type], prefix: '/Game', count: 200 });
+    if (asset)
+      this.props.onChange?.(asset.Path);
+  }
+
+  render() {
+    const { value, label, browse, reset } = this.props;
+    const url = _api.assets.thumbnailUrl(value);
+
+    return (
+      <div className="asset-item">
+        <div className="label">
+          {label}
+        </div>
+        <div className="image">
+          <SafeImage src={url}
+                     fallback="/images/assets/Asset.png"
+                     alt="Asset"
+                     className="box-image" />
+        </div>
+        <div className="input-title">
+          <input readOnly value={value} />
+          {browse &&
+            <span className="asset-icon" onClick={this.onBrowse}><FontAwesomeIcon icon={['fas', 'folder']} /></span>
+          }
+        </div>
+        {reset &&
+          <FontAwesomeIcon icon={['fas', 'undo']} onClick={() => this.props.onChange?.()} />
+        }
+      </div>
+    );
+  }
+};
\ No newline at end of file
diff --git a/Engine/Plugins/VirtualProduction/RemoteControlWebInterface/WebApp/Client/src/components/Widgets/SlidersWidget.tsx b/Engine/Plugins/VirtualProduction/RemoteControlWebInterface/WebApp/Client/src/components/Widgets/SlidersWidget.tsx
--- a/Engine/Plugins/VirtualProduction/RemoteControlWebInterface/WebApp/Client/src/components/Widgets/SlidersWidget.tsx
+++ b/Engine/Plugins/VirtualProduction/RemoteControlWebInterface/WebApp/Client/src/components/Widgets/SlidersWidget.tsx
@@ -10,16 +10,18 @@ type Props = {
   label?: React.ReactNode;
   min?: number;
   max?: number;
+  proportionally?: boolean;
   value?: any;
   
-  onChange?: (widget: ICustomStackProperty, axis?: string, axisValue?: number, locked?: boolean) => any;
+  onChange?: (widget: ICustomStackProperty, axis?: string, axisValue?: number, proportionally?: boolean, min?: number, max?: number) => any;
   onPrecisionModal?: (property: string) => void;
+  onProportionallyToggle?: (property: string, value: string) => void;
 }
 
 export class SlidersWidget extends React.Component<Props> {
 
   render() {   
-    const { widget, label = '', min, max, value } = this.props;
+    const { widget, label = '', min, max, proportionally, value } = this.props;
 
     const propertyType = widget?.propertyType;
     const properties = WidgetUtilities.getPropertyKeys(propertyType);
@@ -40,7 +42,7 @@ export class SlidersWidget extends React.Component<Props> {
                             max={max}
                             precision={precision}
                             value={value?.[property]}
-                            onChange={value => this.props.onChange?.(widget, property, value) || null} />
+                            onChange={value => this.props.onChange?.(widget, property, value, proportionally, min, max) || null} />
                 {isSlider &&
                   <>
                     <div className="limits">{min?.toFixed(1)}</div>
@@ -48,7 +50,7 @@ export class SlidersWidget extends React.Component<Props> {
                             min={min}
                             max={max}
                             showLabel={false}
-                            onChange={value => this.props.onChange?.(widget, property, value) || null} />
+                            onChange={value => this.props.onChange?.(widget, property, value, proportionally, min, max) || null} />
                     <div className="limits">{max?.toFixed(1)}</div>
                   </>
                 }
@@ -56,6 +58,9 @@ export class SlidersWidget extends React.Component<Props> {
             }
           </React.Fragment>
         )}
+        <FontAwesomeIcon icon={['fas', (proportionally) ? 'lock' : 'lock-open']}
+                         className='proportional icon'
+                         onClick={() => this.props.onProportionallyToggle(widget.property, proportionally ? '0' : '1')} />
         <FontAwesomeIcon icon={['fas', 'undo']} className="reset-sliders" onClick={() => this.props.onChange?.(widget)} />
       </div>
     );
diff --git a/Engine/Plugins/VirtualProduction/RemoteControlWebInterface/WebApp/Client/src/components/Widgets/VectorWidget.tsx b/Engine/Plugins/VirtualProduction/RemoteControlWebInterface/WebApp/Client/src/components/Widgets/VectorWidget.tsx
--- a/Engine/Plugins/VirtualProduction/RemoteControlWebInterface/WebApp/Client/src/components/Widgets/VectorWidget.tsx
+++ b/Engine/Plugins/VirtualProduction/RemoteControlWebInterface/WebApp/Client/src/components/Widgets/VectorWidget.tsx
@@ -11,14 +11,15 @@ type Props = {
   label?: React.ReactNode;
   value?: any;
   vector?: ICustomStackProperty;
-  locked?: string[];
   min?: number;
   max?: number;
+  proportionally?: boolean;
 
-  onAxisChange?: (widget: ICustomStackProperty, axis: string, axisValue: number, locked?: boolean) => void;
+  onAxisChange?: (widget: ICustomStackProperty, axis: string, axisValue: number, proportionally?: boolean, min?: number, max?: number) => void;
   onToggleWidgetLock?: (key: string) => void;
   onToggleVectorDrawer?: (vector: ICustomStackProperty) => void;
   onSetVector?: (vector: ICustomStackProperty) => void;
+  onProportionallyToggle?: (property: string, value: string) => void;  
 }
 
 const modes = [WidgetTypes.Joystick, WidgetTypes.Dial, WidgetTypes.Sliders];
@@ -38,7 +39,7 @@ export class VectorWidget extends React.Component<Props> {
   }
   
   render() {
-    const { widget, label = '', value, locked, vector, min, max } = this.props;
+    const { widget, label = '', value, vector, min, max, proportionally } = this.props;
     const key = widget.id;
     const keys = WidgetUtilities.getPropertyKeys(widget.propertyType);
     const widgets = _.compact(modes.map(k => widget.widgets?.find(w => w === k))) || [];
@@ -54,9 +55,13 @@ export class VectorWidget extends React.Component<Props> {
                         max={max}
                         precision={WidgetUtilities.getPropertyPrecision(widget.propertyType)}
                         value={value?.[property]}
-                        onChange={value => this.props.onAxisChange?.(widget, property, value, locked?.includes(key))} />
+                        onChange={value => this.props.onAxisChange?.(widget, property, value, proportionally, min, max)} />
           </div>
         )}
+        { }
+        <FontAwesomeIcon icon={['fas', (proportionally) ? 'lock' : 'lock-open' ]}
+                         className={`proportional ${vector?.id === key ? 'selected' : ''}`}
+                         onClick={() => this.props.onProportionallyToggle(widget.property, proportionally ? '0' : '1')} />
 
         {widget.propertyType && widgets.length > 0 &&
           <FontAwesomeIcon icon={['fas', 'gamepad']}
diff --git a/Engine/Plugins/VirtualProduction/RemoteControlWebInterface/WebApp/Client/src/components/Widgets/index.ts b/Engine/Plugins/VirtualProduction/RemoteControlWebInterface/WebApp/Client/src/components/Widgets/index.ts
--- a/Engine/Plugins/VirtualProduction/RemoteControlWebInterface/WebApp/Client/src/components/Widgets/index.ts
+++ b/Engine/Plugins/VirtualProduction/RemoteControlWebInterface/WebApp/Client/src/components/Widgets/index.ts
@@ -1,3 +1,4 @@
+export * from './AssetWidget';
 export * from './DropdownWidget';
 export * from './ImageSelectorWidget';
 export * from './SlidersWidget';
diff --git a/Engine/Plugins/VirtualProduction/RemoteControlWebInterface/WebApp/Client/src/components/controls/ColorPicker.tsx b/Engine/Plugins/VirtualProduction/RemoteControlWebInterface/WebApp/Client/src/components/controls/ColorPicker.tsx
--- a/Engine/Plugins/VirtualProduction/RemoteControlWebInterface/WebApp/Client/src/components/controls/ColorPicker.tsx
+++ b/Engine/Plugins/VirtualProduction/RemoteControlWebInterface/WebApp/Client/src/components/controls/ColorPicker.tsx
@@ -35,6 +35,8 @@ type Props = {
   parts?: ColorPickerParts[];
   widget?: string;
   alpha?: boolean;
+  max?: number;
+  reset?: boolean;
 
   onChange?: (value?: ColorPickerValue) => void;
   onPrecisionModal?: () => void;
@@ -59,6 +61,7 @@ export class ColorPicker extends React.Component<Props, State> {
     ],
     widget: WidgetTypes.ColorPicker,
     mode: ColorMode.Rgb,
+    reset: true,
   };
 
   state: State = {
@@ -77,10 +80,16 @@ export class ColorPicker extends React.Component<Props, State> {
   }
 
   componentDidUpdate(prevProps: Props) {
-    const { value, type } = this.props;
+    const { value, type, max } = this.props;
 
     if (prevProps.value !== value || type !== prevProps.type)
       this.updateValues();
+
+    if (prevProps.max !== max)
+      this.setState({ max });
+
+    if (this.props.mode !== prevProps.mode)
+      this.forceUpdate();
   }
 
   updateValues = () => {
@@ -96,7 +105,7 @@ export class ColorPicker extends React.Component<Props, State> {
 
     let rgb = color;
     if (type === PropertyType.Vector4 || type === PropertyType.LinearColor) {
-      max = 1;
+      max = this.props.max ?? 1;
 
       for (const key in color)
         if (color[key] > max)
@@ -170,7 +179,7 @@ export class ColorPicker extends React.Component<Props, State> {
     const y = e.clientY - yCenter;
 
     const radius = Math.sqrt(x * x + y * y);
-    const maxRadius = this.circleRef.current.clientHeight / 2;    
+    const maxRadius = this.circleRef.current.clientHeight / 2;
 
     const radians = Math.atan2(x, y);
     const h = this.normalize(-1 * radians * (180 / Math.PI) - 90) / 360;
@@ -280,12 +289,10 @@ export class ColorPicker extends React.Component<Props, State> {
     const { type } = this.props;
 
     let rgb = color as ColorProperty;
+
     if (property === 'v')
       v = value;
 
-    if (property === 's')
-      value = Math.max(value, 0.01);
-
     if (property === 'a' && type === PropertyType.Color)
       value *= 255;
 
@@ -294,8 +301,8 @@ export class ColorPicker extends React.Component<Props, State> {
 
     const hsv = WidgetUtilities.rgb2Hsv(rgb);
 
+    value = Math.max(value, 0.00001);
     hsv[property] = value;
-    hsv.v = Math.max(v, 0.00001);
 
     rgb = WidgetUtilities.hsv2rgb(hsv);
     color = rgb;
@@ -361,9 +368,9 @@ export class ColorPicker extends React.Component<Props, State> {
     return { background: `rgba(${rgb?.R}, ${rgb?.G}, ${rgb?.B}, ${rgb?.A ?? 1})` };
   }
 
-  getPointerPosition = (): { circleX: number, circleY: number, rectangleX: number, rectangleY: number } => {
+  getPointerPosition = (): { circleX?: number, circleY?: number, rectangleX?: number, rectangleY?: number } => {
     const { color, max } = this.state;
-    const { type } = this.props;
+    const { type, mode } = this.props;
 
     let rgb: ColorProperty = null;
 
@@ -381,13 +388,24 @@ export class ColorPicker extends React.Component<Props, State> {
     const hsv = WidgetUtilities.rgb2Hsv(rgb);
     const radian = (hsv.h - 0.25) * Math.PI * 2;
 
-    const circleX = Math.sin(radian) * Math.max(0, Math.min(hsv.s, 1)) * 90;
-    const circleY = -Math.cos(radian) * Math.max(0, Math.min(hsv.s, 1)) * 90;
+    if (mode === ColorMode.Rgb) {
+      const maxCircleRadius = this.circleRef.current?.clientHeight / 2;
+
+      const circleX = Math.sin(radian) * Math.max(0, Math.min(hsv.s, 1)) * (maxCircleRadius - 10);
+      const circleY = -Math.cos(radian) * Math.max(0, Math.min(hsv.s, 1)) * (maxCircleRadius - 10);
+
+      return { circleX, circleY };
+    }
+
+    const rectangleRef = this.rectangleRef.current;
 
-    const rectangleX = Math.max(0, Math.min(hsv.s, 1)) * (210 - 19);
-    const rectangleY = (1 - Math.max(0, hsv.v)) * (220 - 19);
+    const rectangleWidth = rectangleRef?.clientWidth ?? 0;
+    const rectangleHeight = rectangleRef?.clientHeight ?? 0;
 
-    return { circleX, circleY, rectangleX, rectangleY };
+    const rectangleX = Math.max(0, Math.min(hsv.s, 1)) * (rectangleWidth - 19);
+    const rectangleY = Math.max(0, (1 - Math.max(0, hsv.v)) * (rectangleHeight - 19));
+  
+    return { rectangleX, rectangleY };
   }
 
   getVisibleParts = () => {
@@ -472,7 +490,7 @@ export class ColorPicker extends React.Component<Props, State> {
     const { mode, type } = this.props;
     const { color, max } = this.state;
     const pos = this.getPointerPosition();
-    const style: React.CSSProperties = { transform: `translate(${pos.circleX}px, ${pos.circleY}px)` };    
+    const style: React.CSSProperties = { transform: `translate(${pos.circleX}px, ${pos.circleY}px)` };
 
     if (mode === ColorMode.Rgb)
       return (
@@ -499,7 +517,7 @@ export class ColorPicker extends React.Component<Props, State> {
     if (rgb?.A && type === PropertyType.Color)
       rgb.A /= 255;      
 
-    const rectangleStyle: React.CSSProperties = { backgroundColor: `rgba(${rgb.R}, ${rgb.G}, ${rgb.B}, ${rgb?.A ?? 1})`, width: '210px', height: '220px' };
+    const rectangleStyle: React.CSSProperties = { backgroundColor: `rgba(${rgb.R}, ${rgb.G}, ${rgb.B}, ${rgb?.A ?? 1})` };
     const rectanglePointerStyle: React.CSSProperties = { transform: `translate(${pos.rectangleX}px, ${pos.rectangleY}px)` };    
 
     return (
@@ -515,7 +533,7 @@ export class ColorPicker extends React.Component<Props, State> {
   }
 
   render() {
-    const { label, type, widget, value, alpha } = this.props;
+    const { label, type, widget, value, alpha, reset } = this.props;
     const { color, max } = this.state;
 
     const visible = this.getVisibleParts();
@@ -536,44 +554,47 @@ export class ColorPicker extends React.Component<Props, State> {
         <div className="color-picker-wheel">
           <div>{label}</div>
           {!!visible[ColorPickerParts.Circle] && this.rendeCircleRectangle()}
-          {!!visible[ColorPickerParts.Wheel] && (
-            <div className="slider-wheel-container">
-              {widget === WidgetTypes.ColorPicker && (
-                <div className="slider-wheel-top">
-                  <span className="title">Hue</span>
-                  <span className="value">{(hsv.h * 360).toFixed(1)}°</span>
+          <div className='color-picker-controls'>
+            {!!visible[ColorPickerParts.Wheel] && (
+              <div className="slider-wheel-container">
+                {widget === WidgetTypes.ColorPicker && (
+                  <div className="slider-wheel-top">
+                    <span className="title">Hue</span>
+                    <span className="value">{(hsv.h * 360).toFixed(1)}°</span>
+                  </div>
+                )}
+                <SliderWheel onWheelMove={this.onHsvWheelMove} />
+              </div>
+            )}
+            {!!visible[ColorPickerParts.Saturation] && (
+              <ColorPickerSlider value={hsv.s}
+                                 label="Saturation"
+                                 onChange={this.onValueSliderChange.bind(this, 's')} />
+            )}
+            {!!visible[ColorPickerParts.Value] && (
+              <ColorPickerSlider value={hsv.v}
+                                 label="Value"
+                                 onChange={this.onValueSliderChange.bind(this, 'v')} />
+            )}
+            {!!alpha && !!visible[ColorPickerParts.Alpha] && (
+              <ColorPickerSlider value={a}
+                                 label="Alpha"
+                                 onChange={this.onValueSliderChange.bind(this, 'a')} />
+            )}
+            
+            <div className="color-picker-inputs">
+              {!!visible[ColorPickerParts.Inputs] && (
+                <div className='rgb-inputs'>
+                  {this.renderRgbInput(SliderType.R)}
+                  {this.renderRgbInput(SliderType.G)}
+                  {this.renderRgbInput(SliderType.B)}
                 </div>
               )}
-              <SliderWheel onWheelMove={this.onHsvWheelMove} />
-            </div>
-          )}
-          {!!visible[ColorPickerParts.Saturation] && (
-            <ColorPickerSlider value={hsv.s}
-                               label="Saturation"
-                               onChange={this.onValueSliderChange.bind(this, 's')} />
-          )}
-          {!!visible[ColorPickerParts.Value] && (
-            <ColorPickerSlider value={hsv.v}
-                               label="Value"
-                               onChange={this.onValueSliderChange.bind(this, 'v')} />
-          )}
-          {!!alpha && !!visible[ColorPickerParts.Alpha] && (
-            <ColorPickerSlider value={a}
-                               label="Alpha"
-                               onChange={this.onValueSliderChange.bind(this, 'a')} />
-          )}
-          <div className="color-picker-inputs">
-            {!!visible[ColorPickerParts.Inputs] && (
-              <>
-                {this.renderRgbInput(SliderType.R)}
-                {this.renderRgbInput(SliderType.G)}
-                {this.renderRgbInput(SliderType.B)}
-              </>
-            )}
-            {widget === WidgetTypes.ColorPicker && !!value && <div className="color" style={this.getPreviewColorStyles()} />}
+              {widget === WidgetTypes.ColorPicker && !!value && <div className="color" style={this.getPreviewColorStyles()} />}
+            </div>  
           </div>
         </div>
-        <FontAwesomeIcon icon={['fas', 'undo']} onClick={() => this.props.onChange?.()} />
+        {reset && <FontAwesomeIcon icon={['fas', 'undo']} onClick={() => this.props.onChange?.()} />}
       </div>
     );
   }
diff --git a/Engine/Plugins/VirtualProduction/RemoteControlWebInterface/WebApp/Client/src/components/controls/DialsWrapper.tsx b/Engine/Plugins/VirtualProduction/RemoteControlWebInterface/WebApp/Client/src/components/controls/DialsWrapper.tsx
--- a/Engine/Plugins/VirtualProduction/RemoteControlWebInterface/WebApp/Client/src/components/controls/DialsWrapper.tsx
+++ b/Engine/Plugins/VirtualProduction/RemoteControlWebInterface/WebApp/Client/src/components/controls/DialsWrapper.tsx
@@ -2,6 +2,7 @@ import React from 'react';
 import { Dial, DialMode } from '.';
 import { WidgetUtilities } from 'src/utilities';
 import { JoystickValue, PropertyType } from 'src/shared';
+import { FontAwesomeIcon } from '@fortawesome/react-fontawesome';
 
 
 type Props = {
@@ -15,7 +16,11 @@ type Props = {
   type?: PropertyType;
   hidePrecision?: boolean;
   label?: React.ReactNode;
+  proportionally?: boolean;
+  property?: string;
+
   onChange?: (res:  number | JoystickValue) => void;
+  onProportionallyToggle?: (property: string, value: string) => void;
 };
 
 export class DialsWrapper extends React.Component<Props> {
@@ -31,20 +36,46 @@ export class DialsWrapper extends React.Component<Props> {
   }
 
   onChange = (res: number, key?: string) => {
-    const { onChange } = this.props;
-    let { value } = this.props;
+    let { value, proportionally, min, max } = this.props;
     if (value === undefined) {
       const dials = this.getProperties();
       value = dials.length ? {} : 0;
     }
 
-    if (key) {
-      value[key] = res;
-    } else {
-      value = res;
+    let prev = value[key];
+    if (prev === 0 || prev === undefined)
+      prev = 1;
+
+    const ratio = Math.max(0.001, res) / prev;
+
+    if (proportionally && !isNaN(ratio)) {
+      for (const key of Object.keys(value)) {
+        let val = value[key] * ratio;
+        if (!isNaN(min) && !isNaN(max))
+          val = Math.min(max, Math.max(min, val));
+
+        value[key] = val;
+      }
+    }
+    else {
+      if (key)
+        value[key] = res;
+      else
+        value = res;
     }
 
-    onChange?.(value);
+    this.props.onChange?.(value);
+  }
+
+  renderProportionallyIcon = () => {
+    const { proportionally, property } = this.props;
+
+    if (proportionally === undefined)
+      return null;
+
+    return <FontAwesomeIcon icon={['fas', (proportionally) ? 'lock' : 'lock-open']}
+                            className='proportional icon'
+                            onClick={() => this.props.onProportionallyToggle(property, proportionally ? '0' : '1')} />;
   }
 
   render() {
@@ -73,6 +104,7 @@ export class DialsWrapper extends React.Component<Props> {
     return (
       <div className="dial-wrapper-container">
         {label}
+        {this.renderProportionallyIcon()}
         <div className="dial-wrapper-block">
           {!!dials.length &&
             dials.map(key => {
diff --git a/Engine/Plugins/VirtualProduction/RemoteControlWebInterface/WebApp/Client/src/components/controls/Slider.tsx b/Engine/Plugins/VirtualProduction/RemoteControlWebInterface/WebApp/Client/src/components/controls/Slider.tsx
--- a/Engine/Plugins/VirtualProduction/RemoteControlWebInterface/WebApp/Client/src/components/controls/Slider.tsx
+++ b/Engine/Plugins/VirtualProduction/RemoteControlWebInterface/WebApp/Client/src/components/controls/Slider.tsx
@@ -7,15 +7,22 @@ type Props = {
   max?: number;
   size?: number;
   value?: number;
-  onChange?: (value: number) => void;
   showGrid?: boolean;
   showLimits?: boolean;
   showLabel?: boolean;
   precision?: number;
+  touchMode?: boolean;
+  style?: React.CSSProperties;
+  
+  onChange?: (value: number) => void;
 };
 
 export class Slider extends React.Component<Props> {
 
+  static defaultProps: Props = {
+    style: {},
+  };
+
   ref = React.createRef<HTMLDivElement>();
   circleRef = React.createRef<HTMLDivElement>();
   monitoring: boolean = false;
@@ -23,6 +30,12 @@ export class Slider extends React.Component<Props> {
   dragOffset: number = null;
 
   onDown = (e: React.PointerEvent<HTMLDivElement>) => {
+    const { touchMode } = this.props;
+    const target = e.target as HTMLDivElement;
+
+    if (touchMode && !target.classList.contains('circle'))
+      return;
+
     this.monitoring = true;
     this.ref.current.setPointerCapture(e.pointerId);
 
@@ -87,10 +100,15 @@ export class Slider extends React.Component<Props> {
   }
 
   render()  {
-    const { min, max, showGrid, showLimits, showLabel, precision } = this.props;
+    const { min, max, showGrid, showLimits, showLabel, precision, touchMode, style } = this.props;
     const value = this.dragValue ?? this.props.value;
 
     let percent = 0;
+    let className = 'slider ';
+
+    if (touchMode)
+      className += 'touch-mode';
+
     if (!isNaN(value)) {
       percent = (value - min) / (max - min) * 100;
       percent = Math.max(0, Math.min(percent, 100));
@@ -122,10 +140,11 @@ export class Slider extends React.Component<Props> {
                 onPointerUp={this.onUp}
                 onPointerMove={this.onMove}>
             <div className="slider-block">
-              <div className="slider" />
+              <div className={className} style={style} />
               <div className="range-fill" style={{ width: `${percent}%` }} />
-              <div className='circle-wrapper' style={{ transform: `translateX(${percent}%)` }}>
-                <div className="circle" ref={this.circleRef} />
+              <div className="circle-wrapper" style={{ transform: `translateX(${percent}%)` }}>
+                <div className="circle" 
+                     ref={this.circleRef} />
               </div>
             </div>
 
diff --git a/Engine/Plugins/VirtualProduction/RemoteControlWebInterface/WebApp/Client/src/components/controls/SliderWheel.tsx b/Engine/Plugins/VirtualProduction/RemoteControlWebInterface/WebApp/Client/src/components/controls/SliderWheel.tsx
--- a/Engine/Plugins/VirtualProduction/RemoteControlWebInterface/WebApp/Client/src/components/controls/SliderWheel.tsx
+++ b/Engine/Plugins/VirtualProduction/RemoteControlWebInterface/WebApp/Client/src/components/controls/SliderWheel.tsx
@@ -2,6 +2,11 @@ import React from 'react';
 
 
 type SliderWheelProps = {
+  vertical?: boolean;
+  size?: number,
+  className?: string;
+  style?: React.CSSProperties;
+
   onWheelMove?: (value: number, offset?: number) => void;
   onWheelStart?: () => void;
 }
@@ -12,6 +17,10 @@ type SliderWheelState = {
 
 export class SliderWheel extends React.Component<SliderWheelProps, SliderWheelState> {
 
+  static defaultProps: SliderWheelProps = {
+    style: {},
+  };
+
   state: SliderWheelState = {
     offset: 0,
   };
@@ -25,9 +34,14 @@ export class SliderWheel extends React.Component<SliderWheelProps, SliderWheelSt
     if (!this.ref.current)
       return;
 
+    const { vertical } = this.props;  
 
     this.monitoring = true;
     this.last = e.clientX;
+
+    if (vertical)
+      this.last = e.clientY;
+
     this.ref.current.setPointerCapture(e.pointerId);
     this.props.onWheelStart?.();
   }
@@ -36,11 +50,19 @@ export class SliderWheel extends React.Component<SliderWheelProps, SliderWheelSt
     if (!this.monitoring)
       return;
 
-    const delta = e.clientX - this.last;
+    const { vertical } = this.props;
+    let delta = e.clientX - this.last;
+
+    if (vertical)
+      delta = this.last - e.clientY;
+
     if (Math.abs(delta) < 2)
       return; 
 
     this.last = e.clientX;
+    if(vertical)
+      this.last = e.clientY;
+
     let { offset } = this.state;
 
     this.sum += delta;
@@ -64,9 +86,12 @@ export class SliderWheel extends React.Component<SliderWheelProps, SliderWheelSt
   }
 
   renderCircles = () => {
+    let { size } = this.props;
     const circles = [];
 
-    for (let i = 0; i < 80; i++)
+    size = size ?? 80;
+
+    for (let i = 0; i < size; i++)
       circles.push(<div key={i} className="slider-circle" />);
 
     return circles;
@@ -74,14 +99,23 @@ export class SliderWheel extends React.Component<SliderWheelProps, SliderWheelSt
 
   render() {
     const { offset } = this.state;
+    let { vertical, className = '', } = this.props;
     const style: React.CSSProperties = { transform: `translateX(${offset}px)` };
 
+    if (vertical)
+      style.transform = `translateY(${offset}px)`;
+
+    className += ' color-picker-slider-wheel ';
+    if (vertical)
+      className += 'vertical';
+
     return (
-      <div className="color-picker-slider-wheel"
+      <div className={className}
            onPointerMove={this.onPointerMove}
            onPointerDown={this.onPointerDown}
            onPointerUp={this.onPointerUp}
-           ref={this.ref}>
+           ref={this.ref}
+           style={this.props.style}>
         <div className="circles-list" style={style}>
           {this.renderCircles()}
         </div>
diff --git a/Engine/Plugins/VirtualProduction/RemoteControlWebInterface/WebApp/Client/src/components/controls/Tabs.tsx b/Engine/Plugins/VirtualProduction/RemoteControlWebInterface/WebApp/Client/src/components/controls/Tabs.tsx
new file mode 100644
--- /dev/null
+++ b/Engine/Plugins/VirtualProduction/RemoteControlWebInterface/WebApp/Client/src/components/controls/Tabs.tsx
@@ -0,0 +1,173 @@
+import React from 'react';
+import _ from 'lodash';
+import { IconProp } from '@fortawesome/fontawesome-svg-core';
+import { FontAwesomeIcon } from '@fortawesome/react-fontawesome';
+
+
+type Props = {
+  children: React.ReactElement[];
+
+  leftIcon?: React.ReactElement[] | React.ReactElement;
+  rightIcon?: React.ReactElement[] | React.ReactElement;
+
+  defaultActiveKey?: string;
+  onlyHeader?: boolean;
+
+  onTabChange?: (active?: string) => void;
+}
+
+type State = {
+  active: string;
+}
+
+export class Tabs extends React.Component<Props, State> {
+
+  state: State = {
+    active: null,
+  }
+
+  selectorRef = React.createRef<HTMLDivElement>();
+  containerRef = React.createRef<HTMLDivElement>();
+
+  componentDidMount() {
+    this.updateActive();
+
+    window.addEventListener('resize', this.updateSelector);
+  }
+
+  componentWillUnmount() {
+    window.removeEventListener('resize', this.updateSelector);
+  }
+
+  componentDidUpdate(prevProps: Readonly<Props>, prevState: Readonly<State>): void {
+    const { active } = this.state;
+    const { defaultActiveKey, children } = this.props;
+
+    if (children.length !== prevProps.children.length)
+      this.updateSelector();
+
+    if (active !== prevState.active)
+      this.updateSelector();
+
+    if (defaultActiveKey !== prevProps.defaultActiveKey)
+      this.updateActive();
+  }
+
+  getSelectorPosition = () => {
+    const elements = this.containerRef.current.getElementsByClassName('tab');
+    if (!elements.length)
+      return;
+
+    let left = 5;
+    let width = null;
+    
+    for(let i = 0; i < elements.length; i++) {
+      const el = elements[i] as HTMLDivElement;
+
+      if (el.classList.contains('active')) {
+        width = ((el.offsetLeft - left) * 2) + el.clientWidth;
+        break;
+      }
+
+      left = ((el.offsetLeft * 2) + el.clientWidth) - left;
+    }
+
+    return [left, width];
+  }
+
+  onTabChange = (active: string) => {
+    this.setState({ active });
+    this.props.onTabChange?.(active);
+  }
+
+  updateActive = () => {
+    const { children, defaultActiveKey } = this.props;
+    const child = _.first(children) as any;
+
+    if (!child)
+      return;
+
+    const id = child.props.id;
+    this.setState({ active: defaultActiveKey ?? id });
+  }
+
+  updateSelector = () => {
+    const selector = this.selectorRef.current;
+    const [left, width] = this.getSelectorPosition();
+
+    selector.style.left = `${left}px`;
+    selector.style.width = `${width}px`;
+  }
+
+  renderTabHeader = (child, index: number) => {
+    if (!child)
+      return null;
+
+    const { active } = this.state;
+    const { id, tab, icon, view } = child.props;
+
+    const isActive = active === id || !active && index === 0;
+    let className = 'tab ';
+
+    if (isActive)
+      className += 'active ';
+
+    return (
+      <div key={id}
+           className={className}
+           onClick={() => this.onTabChange(id)}
+           title={tab}>
+        {!!icon && <FontAwesomeIcon className="icon" icon={icon} />}
+        {tab}
+        {!!view && view({ ...child.props, active: isActive })}
+      </div>
+    );
+  }
+
+  renderTabContent = () => {
+    const { active } = this.state;
+    const { children } = this.props;
+
+    const child = _.find(children, c => c?.props?.id === active);
+
+    if (!child)
+      return null;
+
+    return child;
+  }
+
+  render() {
+    let { children, leftIcon, rightIcon, onlyHeader } = this.props;
+    children = Array.isArray(children) ? children : [children];
+
+    return (
+      <div ref={this.containerRef} className="tabs-container">
+        <div className="tabs-header-elements">
+          {!!leftIcon && leftIcon}
+          <div className="tabs-elements">
+            <div className="tabs-header">
+              {children.map?.(this.renderTabHeader)}
+            </div>
+            <div className="selector" ref={this.selectorRef} />
+          </div>
+          {!!rightIcon && rightIcon}
+        </div>
+        {!onlyHeader && <div className="tabs-content">{this.renderTabContent()}</div>}
+      </div>
+    );
+  }
+};
+
+type TabPaneProps = {
+  tab: string;
+  id: string;
+  icon?: IconProp;
+
+  view?: ({ id, tab, icon, active }) => React.ReactElement;
+}
+
+export class TabPane extends React.Component<TabPaneProps> {
+  render() {
+    return this.props.children;
+  }
+}
\ No newline at end of file
diff --git a/Engine/Plugins/VirtualProduction/RemoteControlWebInterface/WebApp/Client/src/components/controls/TimeControl.tsx b/Engine/Plugins/VirtualProduction/RemoteControlWebInterface/WebApp/Client/src/components/controls/TimeControl.tsx
new file mode 100644
--- /dev/null
+++ b/Engine/Plugins/VirtualProduction/RemoteControlWebInterface/WebApp/Client/src/components/controls/TimeControl.tsx
@@ -0,0 +1,158 @@
+import { FontAwesomeIcon } from '@fortawesome/react-fontawesome';
+import React from 'react';
+
+
+type Props = {
+  label?: string;
+  value?: number;
+  disabled?: boolean;
+
+  onChange?: (value: number) => void;
+}
+
+type State = {
+  value?: number;
+}
+
+enum TimeKey {
+  Hours   = 'HOURS',
+  Minutes = 'MINUTES',
+}
+
+export class TimeControl extends React.Component<Props, State> {
+
+  static defaultProps: Props = {
+    label: '',
+    value: 23.5,
+  };
+
+  state: State = {
+    value: null,
+  }
+
+  hoursInputRef = React.createRef<HTMLInputElement>();
+  minutesInputRef = React.createRef<HTMLInputElement>();
+
+  componentDidMount() {
+    const { value } = this.props;
+    this.setState({ value });
+  }
+
+  componentDidUpdate(prevProps: Props): void {
+    // const { value } = this.props;
+    
+    // if (value !== prevProps.value) {
+    //   this.setState({ value });
+    // }
+  }
+
+  getTwoDigits = (num: number) => {
+    return num.toString().padStart(2, '0');
+  }
+
+  getInputValues = () => {
+    const { value } = this.state;
+
+    if (value === null)
+      return ':';
+
+    const hours = Math.floor(value);
+    const minutes = Math.round((value * 60) % 60);
+
+    return [this.getTwoDigits(hours), this.getTwoDigits(minutes)];
+  }
+
+  normalizeValue = (value: number) => {
+    if (value >= 13)
+      value -= 12;
+
+    if (value < 1)
+      value += 12;
+
+    return +value.toFixed(2);
+  }
+
+  onTimeChange = (dir: number) => {
+    let { value } = this.state;
+
+    value += dir * 0.25;
+    value = this.normalizeValue(value);
+
+    this.props.onChange?.(value);
+    this.setState({ value });
+  }
+
+  onInputClick = (ref: React.RefObject<HTMLInputElement>) => {
+    if (!ref?.current)
+      return;
+
+    ref.current.select();
+  }
+
+  onInputChange = (v: string, key: TimeKey) => {
+    let { value } = this.state;
+    const [h] = value.toString().split('.');
+    const m = (value - +h).toFixed(2);
+
+    let diff = 0;
+
+    switch (key) {
+      case TimeKey.Hours:
+        diff = +v - +h;
+        break;
+
+      case TimeKey.Minutes:
+        if (v.length > 2)
+          v = v.slice(1, 3);
+
+        const d = Math.min(59, +v);
+        diff = +(d / 60).toFixed(2) - +m;
+        break;
+    }
+
+    value = +(value + diff).toFixed(2);
+    if (value >= 13 || value < 1)
+      return;
+
+    this.props.onChange?.(value);
+    this.setState({ value });
+  }
+
+  render() {
+    const { label, disabled } = this.props;
+    const [hours = '', minutes = ''] = this.getInputValues();
+
+    let className = 'time-control ';
+    if (disabled)
+      className += 'disabled';
+
+    return (
+      <div className={className}>
+        <label>{label}</label>
+        <div className="time-component">
+          <span className="icon" onClick={() => this.onTimeChange(-1)}>
+            <FontAwesomeIcon icon={['fas', 'minus']} />
+          </span>
+          <div className="inputs-row">
+            <input ref={this.hoursInputRef}
+                   className="time-input"
+                   type="number"
+                   value={hours}
+                   onClick={() => this.onInputClick(this.hoursInputRef)}
+                   onChange={e => this.onInputChange(e.target.value, TimeKey.Hours)} />
+            :
+            <input ref={this.minutesInputRef}
+                   className="time-input"
+                   type="number"
+                   value={minutes}
+                   onClick={() => this.onInputClick(this.minutesInputRef)}
+                   onChange={e => this.onInputChange(e.target.value, TimeKey.Minutes)} />
+          </div>
+          <span className="icon" onClick={() => this.onTimeChange(1)}>
+            <FontAwesomeIcon icon={['fas', 'plus']} />
+          </span>
+        </div>
+      </div>
+    );
+  }
+};
\ No newline at end of file
diff --git a/Engine/Plugins/VirtualProduction/RemoteControlWebInterface/WebApp/Client/src/components/controls/ValueInput.tsx b/Engine/Plugins/VirtualProduction/RemoteControlWebInterface/WebApp/Client/src/components/controls/ValueInput.tsx
--- a/Engine/Plugins/VirtualProduction/RemoteControlWebInterface/WebApp/Client/src/components/controls/ValueInput.tsx
+++ b/Engine/Plugins/VirtualProduction/RemoteControlWebInterface/WebApp/Client/src/components/controls/ValueInput.tsx
@@ -3,11 +3,13 @@ import React from 'react';
 
 type Props = {
   value?: number | string;
-  onChange: (value: number) => void;
   draggable?: boolean;
   min?: number;
   max?: number;
   precision?: number;
+  disabled?: boolean;
+
+  onChange: (value: number) => void;
 };
 
 type State = {
@@ -56,7 +58,7 @@ export class ValueInput extends React.Component<Props, State> {
   }
 
   onClick = () => {
-    this.input.current.focus();
+    this.input.current.select();
   }
 
   onDown = (e: React.PointerEvent<HTMLDivElement>) => {
@@ -156,18 +158,24 @@ export class ValueInput extends React.Component<Props, State> {
 
   render() {
     const { value } = this.state;
+    const { disabled } = this.props;
 
     return (
       <span ref={this.ref}
             className="input-field-container"
             onClick={this.onClick}
+            onDoubleClick={this.onClick}
             onPointerDown={this.onDown}
             onPointerMove={this.onMove}
             onPointerUp={this.onUp}
             onLostPointerCapture={this.onUp}>
         <input ref={this.input}
+               disabled={disabled}
                className="input-field"
                value={value ?? ''}
+               type="number"
+               inputMode="numeric"
+               pattern="[0-9]*"
                onChange={this.onManualChange}
                onKeyDown={this.onKeyDown}
                onBlur={this.onChangeValue}
diff --git a/Engine/Plugins/VirtualProduction/RemoteControlWebInterface/WebApp/Client/src/components/controls/VectorControl.tsx b/Engine/Plugins/VirtualProduction/RemoteControlWebInterface/WebApp/Client/src/components/controls/VectorControl.tsx
new file mode 100644
--- /dev/null
+++ b/Engine/Plugins/VirtualProduction/RemoteControlWebInterface/WebApp/Client/src/components/controls/VectorControl.tsx
@@ -0,0 +1,200 @@
+import React from 'react';
+import { PropertyType } from 'src/shared';
+import { WidgetUtilities } from 'src/utilities';
+import { ValueInput } from './ValueInput';
+import { FontAwesomeIcon } from '@fortawesome/react-fontawesome';
+import _ from 'lodash';
+
+
+type Props = {
+  label: string;
+  propertyType: PropertyType;
+  value: any;
+  step: number;
+  defaultValue?: any;
+
+  onChange?: (property: string, value: any) => void;
+}
+
+type State = {
+  step: number;
+  activeKey: string;
+}
+
+const dashArr = Array.from(Array(81));
+
+export class VectorControl extends React.Component<Props, State> {
+
+  static defaultProps: Props = {
+    label: '',
+    propertyType: PropertyType.Vector,
+    value: {},
+    step: 0.1,
+  };
+
+  state: State = {
+    step: null,
+    activeKey: null,
+  }
+
+  sliderRef = React.createRef<HTMLDivElement>();
+  pointerRef = React.createRef<HTMLDivElement>();
+  value = null;
+  monitoring: boolean = false;
+  offset: number = 0;
+  step: number = 0;
+  interval = null;
+
+  componentDidMount() {
+    const { propertyType } = this.props;
+    const keys = WidgetUtilities.getPropertyKeys(propertyType);
+    const activeKey = _.first(keys);
+
+    if (!activeKey)
+      return;
+
+    this.setState({ activeKey });
+  }
+
+  onMouseDown = (e: React.PointerEvent<HTMLDivElement>) => {
+    this.monitoring = true;
+    this.sliderRef.current.setPointerCapture(e.pointerId);
+    this.offset = 0;
+    this.value = { ...this.props.value };
+
+    const pointer = this.pointerRef.current.getBoundingClientRect();
+    if (pointer.left <= e.clientX && e.clientX <= pointer.right)
+      this.offset = e.clientX - (pointer.left + (pointer.width / 2));
+
+    this.onMouseMove(e);
+    this.makeMove();
+
+    this.interval = setInterval(this.makeMove, 100);
+  }
+
+  onMouseUp = (e: React.PointerEvent<HTMLDivElement>) => {
+    if (!this.monitoring)
+      return;
+
+    this.monitoring = false;
+    this.offset = 0;
+    this.value = null;
+    this.step = 0;
+    this.sliderRef.current.releasePointerCapture(e.pointerId);
+
+    clearInterval(this.interval);
+    this.interval = null;
+
+    this.setState({ step: null });
+  }
+
+  onMouseMove = (e: React.PointerEvent<HTMLDivElement>) => {
+    if (!this.monitoring || !this.sliderRef.current)
+      return;
+
+    const slider = this.sliderRef.current;
+
+    const rect = slider.getBoundingClientRect();
+    const x = e.clientX - this.offset;
+    const center = (rect.right - rect.left) / 2;
+    const ratio = Math.max(-1, Math.min(((x - (rect.left + center)) / rect.width) * 2, 1));
+    const step = +ratio.toFixed(2);
+
+    this.step = step * this.props.step;
+    this.setState({ step });
+  }
+
+  onKeySelect = (e: React.MouseEvent, activeKey: string) => {
+    e.stopPropagation();
+    this.setState({ activeKey });
+  }
+
+  makeMove = () => {
+    const { activeKey } = this.state;
+    if (!activeKey || this.value?.[activeKey] === undefined)
+      return;  
+
+    this.value[activeKey] += this.step;
+    this.props.onChange(activeKey, this.value[activeKey]);
+  }
+
+  resetActiveKey = (e: React.MouseEvent) => {
+    this.setState({ activeKey: null });
+  }
+
+  getPointerStyles = (): React.CSSProperties => {
+    const { step } = this.state;
+    const slider = this.sliderRef.current?.getBoundingClientRect?.();
+
+    if (step === null || !slider)
+      return { transition: 'transform .2s' };
+
+    const translate = (slider.width / 2 * step).toFixed();
+    return { transform: `translateX(${translate}px)` };
+  }
+
+  getIsDisabled = () => {
+    const { activeKey } = this.state;
+    return !activeKey;
+  }
+
+  renderPropertyInput = (key: string) => {
+    const { activeKey } = this.state;
+    const { value } = this.props;
+
+    let className = 'key ';
+
+    if (key === activeKey)
+      className += 'active';
+
+    return (
+      <div key={key} 
+           className={className} 
+           onClick={(e) => this.onKeySelect(e, key)}>
+        <div className="property">{key}:</div>
+        <ValueInput value={value[key] ?? 0}
+                    disabled={activeKey !== key}
+                    onChange={value => this.props.onChange(key, value)} />
+      </div>
+    );
+  }
+
+  onReset = () => {
+    const { onChange, defaultValue } = this.props;
+    onChange(null, defaultValue);
+  }
+
+  render() {
+    const { label, propertyType } = this.props;
+    const keys = WidgetUtilities.getPropertyKeys(propertyType);
+
+    let className = 'vector-control ';
+    if (this.getIsDisabled())
+      className += 'disabled';
+
+    return (
+      <div className={className} onClick={this.resetActiveKey}>
+        <div className="title">{label}</div>
+        <div className="vector-item">
+          <div className="keys-row">
+            {keys.map(this.renderPropertyInput)}
+          </div>
+          <div className="vector-container">
+            <div ref={this.sliderRef}
+                 className="slider-control"
+                 onClick={e => e.stopPropagation()}
+                 onPointerDown={this.onMouseDown}
+                 onLostPointerCapture={this.onMouseUp}
+                 onPointerMove={this.onMouseMove}>
+              {dashArr.map((_, i) => <div key={`dash-${i}`} className="dash" />)}
+            </div>
+            <div ref={this.pointerRef} className="pointer" style={this.getPointerStyles()} />
+          </div>
+        </div>
+        <span className="reset-icon" onClick={this.onReset}>
+          <FontAwesomeIcon icon={['fas', 'undo']} />
+        </span>
+      </div>
+    );
+  }
+};
\ No newline at end of file
diff --git a/Engine/Plugins/VirtualProduction/RemoteControlWebInterface/WebApp/Client/src/components/controls/index.ts b/Engine/Plugins/VirtualProduction/RemoteControlWebInterface/WebApp/Client/src/components/controls/index.ts
--- a/Engine/Plugins/VirtualProduction/RemoteControlWebInterface/WebApp/Client/src/components/controls/index.ts
+++ b/Engine/Plugins/VirtualProduction/RemoteControlWebInterface/WebApp/Client/src/components/controls/index.ts
@@ -8,9 +8,12 @@ export * from './JoysticksWrapper';
 export * from './SafeImage';
 export * from './Slider';
 export * from './SlidersWrapper';
+export * from './Tabs';
 export * from './Text';
+export * from './TimeControl';
 export * from './Toggle';
 export * from './ValueInput';
+export * from './VectorControl';
 export * from './VectorDrawer';
 export * from './Search';
 export * from './SliderWheel';
\ No newline at end of file
diff --git a/Engine/Plugins/VirtualProduction/RemoteControlWebInterface/WebApp/Client/src/components/index.ts b/Engine/Plugins/VirtualProduction/RemoteControlWebInterface/WebApp/Client/src/components/index.ts
--- a/Engine/Plugins/VirtualProduction/RemoteControlWebInterface/WebApp/Client/src/components/index.ts
+++ b/Engine/Plugins/VirtualProduction/RemoteControlWebInterface/WebApp/Client/src/components/index.ts
@@ -5,4 +5,6 @@ export * from './PropertiesDrawer';
 export * from './modals';
 export * from './Screen';
 export * from './Login';
-export * from './Tabs';
\ No newline at end of file
+export * from './Tabs';
+export * from './SignalIcon';
+export * from './EmptyTab';
\ No newline at end of file
diff --git a/Engine/Plugins/VirtualProduction/RemoteControlWebInterface/WebApp/Client/src/components/modals/IconModal.tsx b/Engine/Plugins/VirtualProduction/RemoteControlWebInterface/WebApp/Client/src/components/modals/IconModal.tsx
--- a/Engine/Plugins/VirtualProduction/RemoteControlWebInterface/WebApp/Client/src/components/modals/IconModal.tsx
+++ b/Engine/Plugins/VirtualProduction/RemoteControlWebInterface/WebApp/Client/src/components/modals/IconModal.tsx
@@ -128,7 +128,7 @@ export class IconModal extends React.Component<Props, State> {
           <div className="search-wrapper search-position">
             <FontAwesomeIcon icon={['fas', 'search']} />
             <input value={filterText}
-                    className="icon-search"
+                    className="icon-search search-input"
                     placeholder="Search icon" 
                     onChange={e => this.onFilterTextChange(e.target.value)}
                     autoFocus={true}
diff --git a/Engine/Plugins/VirtualProduction/RemoteControlWebInterface/WebApp/Client/src/components/modals/SearchModal.tsx b/Engine/Plugins/VirtualProduction/RemoteControlWebInterface/WebApp/Client/src/components/modals/SearchModal.tsx
new file mode 100644
--- /dev/null
+++ b/Engine/Plugins/VirtualProduction/RemoteControlWebInterface/WebApp/Client/src/components/modals/SearchModal.tsx
@@ -0,0 +1,166 @@
+import React from 'react';
+import ReactDOM from 'react-dom';
+import { IAsset } from 'src/shared';
+import { _api } from '../../reducers';
+import { FontAwesomeIcon } from '@fortawesome/react-fontawesome';
+import { SafeImage } from '../controls';
+import _ from 'lodash';
+
+
+type Props = {
+  placeholder?: string;
+  types?: string[];
+  prefix?: string;
+  count?: number;
+  fallback?: string;
+
+  close: (asset?: IAsset) => void;
+};
+
+type State = {
+  search: string;
+  assets: IAsset[];
+  loading: boolean;
+};
+
+export class SearchModal extends React.Component<Props, State> {
+  state: State = {
+    search: '',
+    assets: [],
+    loading: false,
+  };
+
+  componentDidMount() {
+    document.addEventListener('keydown', this.onKeyPress);
+    this.fetchAssets();
+  }
+
+  componentWillUnmount() {
+    document.removeEventListener('keydown', this.onKeyPress);
+  }
+
+  fetchAssets = async () => {
+    const { search } = this.state;
+    const { types, prefix, count } = this.props;
+
+    this.setState({ loading: true });
+    
+    const assets = await _api.assets.search(search, types, prefix, {}, count);
+    this.setState({ assets, loading: false });
+  }
+
+  fetchAssetsDebounced = _.debounce(this.fetchAssets, 200);
+
+  onKeyPress = (e) => {
+    if (e.key === 'Escape')
+      this.onCancel();
+  }
+
+  private static pending: (asset?: IAsset) => void;
+
+  public static open(props: Partial<Props>): Promise<IAsset> {
+    SearchModal.pending?.(undefined);
+    SearchModal.pending = undefined;
+
+    const div = document.createElement('div');
+    document.getElementById('modal').prepend(div);
+
+    return new Promise<IAsset>(resolve => {
+      SearchModal.pending = resolve;
+      ReactDOM.render(<SearchModal {...props} close={resolve} />, div);
+    })
+      .finally(() => {
+        const unmountResult = ReactDOM.unmountComponentAtNode(div);
+        if (unmountResult && div.parentNode)
+          div.parentNode.removeChild(div);
+      });
+  }
+
+  onOK = () => {    
+    this.props.close();
+  }
+
+  onCancel = () => {
+    this.props.close();
+  }
+
+  onKeyDown = (e: React.KeyboardEvent<HTMLInputElement>) => {
+    switch (e.key) {
+      case 'Enter':
+        return this.onOK();
+
+      case 'Escape':
+        return this.onCancel();
+    }
+  }
+
+  onSearch = (search: string) => {
+    this.setState({ search });
+    this.fetchAssetsDebounced();
+  }
+
+  renderAssetItem = (asset: IAsset, key: number) => {
+    const { fallback } = this.props;
+    const url = _api.assets.thumbnailUrl(asset.Path);
+
+    return (
+      <div className="search-item" key={key} onClick={() => this.props.close(asset)}>
+        <div className="image">
+        <SafeImage src={url} 
+                   fallback={fallback ?? "/images/assets/Asset.png"}
+                   alt={asset.Name}
+                   className="box-image" />
+        </div>
+        <div className="title" title={asset.Name}>
+          {asset.Name}
+        </div>
+        <div className="path" title={asset.Path}>
+          {asset.Path}
+        </div>
+      </div>
+    );
+  }
+
+  render() {
+    const { placeholder } = this.props;
+    const { search, assets, loading } = this.state;
+
+    return (
+      <div className="fullscreen" onPointerDown={this.onCancel}>
+        <div className="search-modal" onPointerDown={e => e.stopPropagation()}>
+          <div className="close-modal" onClick={this.onCancel}>
+            <FontAwesomeIcon icon={['fas', 'times']} />
+          </div>
+          <div className="search-wrapper">
+            <FontAwesomeIcon icon={['fas', 'search']} />
+            <input value={search}
+                   className="icon-search search-input"
+                   placeholder={placeholder}
+                   onChange={e => this.onSearch(e.target.value)}
+                   autoFocus={true}
+                   autoComplete="off" />
+            {search && <FontAwesomeIcon icon={['fas', 'times-circle']}
+                                        className="reset-search-icon"
+                                        onClick={() => this.setState({ search: '' })} />}
+          </div>
+          <div className="search-content">
+            {loading &&
+              <div className="content-placeholder">
+                Searching
+                <FontAwesomeIcon icon={['fas', 'sync-alt']}
+                                 spin
+                                 className="loader-icon" />
+              </div>}
+
+            {!loading && !assets.length &&
+              <div className="content-placeholder">
+                No Data
+              </div>}
+
+            {assets.map(this.renderAssetItem)}
+          </div>
+        </div>
+      </div>
+    );
+  }
+}
\ No newline at end of file
diff --git a/Engine/Plugins/VirtualProduction/RemoteControlWebInterface/WebApp/Client/src/components/modals/index.ts b/Engine/Plugins/VirtualProduction/RemoteControlWebInterface/WebApp/Client/src/components/modals/index.ts
--- a/Engine/Plugins/VirtualProduction/RemoteControlWebInterface/WebApp/Client/src/components/modals/index.ts
+++ b/Engine/Plugins/VirtualProduction/RemoteControlWebInterface/WebApp/Client/src/components/modals/index.ts
@@ -1,4 +1,5 @@
 export * from './AlertModal';
 export * from './IconModal';
+export * from './PrecisionModal';
 export * from './RenameModal';
-export * from './PrecisionModal';
\ No newline at end of file
+export * from './SearchModal';
\ No newline at end of file
diff --git a/Engine/Plugins/VirtualProduction/RemoteControlWebInterface/WebApp/Client/src/reducers/api.ts b/Engine/Plugins/VirtualProduction/RemoteControlWebInterface/WebApp/Client/src/reducers/api.ts
--- a/Engine/Plugins/VirtualProduction/RemoteControlWebInterface/WebApp/Client/src/reducers/api.ts
+++ b/Engine/Plugins/VirtualProduction/RemoteControlWebInterface/WebApp/Client/src/reducers/api.ts
@@ -2,7 +2,7 @@ import { Dispatch } from 'redux';
 import { createAction, createReducer } from 'redux-act';
 import dotProp from 'dot-prop-immutable';
 import io from 'socket.io-client';
-import { IAsset, IPayload, IPayloads, IPreset, IView, PropertyValue, TabLayout } from '../shared';
+import { IAsset, IPayload, IPayloads, IPreset, IView, PropertyValue, TabLayout, IHistory, ConnectionSignal } from '../shared';
 import _ from 'lodash';
 
 
@@ -17,24 +17,35 @@ export type ApiState = {
     keyCorrect?: boolean;
     loading?: boolean;
     isOpen?: boolean;
+    signal?: ConnectionSignal;
   },
+  lockedUI: boolean;
+  undo: IHistory[],
+  redo: IHistory[],
 };
 
 
 let _preset;
 let _dispatch: Dispatch;
+let _getState: () => { api: ApiState };
 let _socket: SocketIOClient.Socket;
 let _passphrase: string;
+let _pingInterval: NodeJS.Timeout;
+
 const _host = (process.env.NODE_ENV === 'development' ? `http://${window.location.hostname}:7001` : '');
 
 function _initialize(dispatch: Dispatch, getState: () => { api: ApiState }) {
   _dispatch = dispatch;
+  _getState = getState;
   _passphrase = localStorage.getItem('passphrase'); 
 
   _socket = io(`${_host}/`, { path: '/api/io' });
 
   _socket
-    .on('disconnect', () => dispatch(API.STATUS({ connected: false, isOpen: false, keyCorrect: false, version: undefined })))
+    .on('disconnect', () => {
+      dispatch(API.STATUS({ connected: false, isOpen: false, keyCorrect: false, version: undefined }));
+      clearInterval(_pingInterval);
+    })
     .on('presets', (presets: IPreset[]) => dispatch(API.PRESETS(presets)))
     .on('payloads', (payloads: IPayloads) => {
       dispatch(API.PAYLOADS(payloads));
@@ -62,18 +73,19 @@ function _initialize(dispatch: Dispatch, getState: () => { api: ApiState }) {
     })
     .on('connected', (connected: boolean, version: string) => {
       dispatch(API.STATUS({ connected, version }));
-
+      if (_pingInterval)
+        clearInterval(_pingInterval);
+      
+      _pingInterval = null;
       if (_passphrase !== null)
         _api.passphrase.login(_passphrase);
-          
-      _api.presets.get();
-      _api.payload.all();
     })
     .on('opened', (isOpen: boolean) => {
       dispatch(API.STATUS({ isOpen, loading: false }));
 
       _api.presets.get();
       _api.payload.all();
+      _pingInterval = setInterval(_api.ping, 1000);
     })
     .on('passphrase', (wrongPassphrase: string) => {
       const isLoginStillCorrect = _passphrase !== wrongPassphrase && _passphrase !== undefined;
@@ -83,6 +95,18 @@ function _initialize(dispatch: Dispatch, getState: () => { api: ApiState }) {
     })
     .on('loading', (loading: boolean) => {
       dispatch(API.STATUS({ loading }));
+    })
+    .on('pong', (pingTime: number, pongTime: number) => {
+      let signal = ConnectionSignal.Good;
+      const diff = pongTime - pingTime;
+
+      if (diff > 50)
+        signal = ConnectionSignal.Bad;
+      else if (diff > 15)
+        signal = ConnectionSignal.Normal;
+
+
+      dispatch(API.STATUS({ signal }));
     });
 }
 
@@ -127,18 +151,49 @@ const API = {
   PAYLOAD: createAction<IPayload>('API_PAYLOAD'),
   PAYLOADS: createAction<IPayloads>('API_PAYLOADS'),
   PAYLOADS_VALUE: createAction<IPayloads>('API_PAYLOADS_VALUE'),
+  LOCK_UI: createAction<boolean>('LOCK_UI'),
+  HISTORY_UPDATE: createAction<{ undo: IHistory[], redo: IHistory[] }>('HISTORY_UPDATE'),
 };
 
 
 export const _api = {
   initialize: () => _initialize.bind(null),
+  ping: () => _socket.emit('client', new Date().getTime()),
+  lockUI: (locked: boolean) => _dispatch(API.LOCK_UI(locked)),
+  undoHistory: () => {
+    const { undo, redo, payload } = _getState().api;
+    const action = undo.shift();
+
+    if (!action)
+      return;
+    
+    redo.push({ property: action.property, value: payload[action.property], time: new Date() });
+
+    _dispatch(API.HISTORY_UPDATE({ undo, redo }));
+    _api.payload.set(action.property, action.value, false);
+  },
+  redoHistory: () => {
+    const { undo, redo, payload } = _getState().api;
+    const action = redo.pop();
+
+    if (!action)
+      return;
+
+    undo.unshift({ property: action.property, value: payload[action.property], time: new Date() });
+
+    _dispatch(API.HISTORY_UPDATE({ undo, redo }));
+    _api.payload.set(action.property, action.value, false);
+  },
 
   presets: {
     get: (): Promise<IPreset[]> => _get('/api/presets', API.PRESETS),
     load: (id: string): Promise<IPreset> => _get(`/api/presets/${id}/load`),
+    favorite: (id: string, value: boolean): Promise<IPreset> => _put(`/api/presets/${id}/favorite`, { value }),
     select: (preset?: IPreset) => {
       _api.presets.load(preset?.ID);
+
       _dispatch(API.PRESET_SELECT(preset?.ID));
+      _dispatch(API.HISTORY_UPDATE({ undo: [], redo: [] }));
     },
   },
   views: {
@@ -146,7 +201,7 @@ export const _api = {
       let view = await _get(`/api/presets/view?preset=${preset}`);
       if (typeof(view) !== 'object' || !view?.tabs?.length) {
         view = {
-          tabs: [{ name: 'Tab 1', layout: TabLayout.Stack, panels: [] }]
+          tabs: [{ name: 'Tab 1', layout: TabLayout.Empty }]
         };
       }
 
@@ -174,7 +229,27 @@ export const _api = {
   payload: {
     get: (preset: string): Promise<IPayload> => _get(`/api/presets/payload?preset=${preset}`, API.PAYLOAD),
     all: (): Promise<IPayloads> => _get('/api/payloads', API.PAYLOADS),
-    set: (property: string, value: PropertyValue) => {
+    set: (property: string, value: PropertyValue, historyPush = true) => {
+      let { undo, redo, payloads } = _getState().api;
+
+      if (historyPush) {
+        if (undo.length > 50)
+          undo.shift();
+
+        if (redo.length > 0)
+          redo = [];
+
+        const lastUndo = _.last(undo);
+        const time = new Date();
+
+        if (lastUndo?.property === property && (time.valueOf() - lastUndo?.time?.valueOf()) < 60 * 1000)
+          lastUndo.value = payloads[_preset][property];
+        else
+          undo.unshift({ property, value: payloads[_preset][property], time });
+
+          _dispatch(API.HISTORY_UPDATE({ undo, redo }));
+      }
+
       _socket.emit('value', _preset, property, value);
     },
     execute: (func: string, args?: Record<string, any>) => {
@@ -197,11 +272,12 @@ export const _api = {
     get: (url: string) => _put('/api/proxy', { method: 'GET', url }),
     put: (url: string, body: any) => _put('/api/proxy', { method: 'PUT', url, body }),
     function: (objectPath: string, functionName: string, parameters: Record<string, any> = {}): Promise<any> => {
-      const body = { objectPath, functionName, parameters };
-      return _api.proxy.put('/remote/object/call', body);
+      const body = { objectPath, functionName, parameters, generateTransaction: true };
+      return _api.proxy.put('/remote/object/call', body)
+                        .then(ret => ret?.ReturnValue);
     },
     property: {
-      get: (objectPath: string, propertyName: string) => {
+      get: <T = {}>(objectPath: string, propertyName?: string): Promise<T> => {
         const body = { 
           objectPath,
           propertyName,
@@ -229,9 +305,13 @@ const initialState: ApiState = {
   view: { tabs: null },
   status: {
     connected: false,
+    signal: ConnectionSignal.Good,
     keyCorrect: false,
     isOpen: false,
   },
+  lockedUI: false,
+  undo: [],
+  redo: [],
 };
 
 const reducer = createReducer<ApiState>({}, initialState);
@@ -306,6 +386,12 @@ reducer
     _api.views.get(preset);
     _api.payload.get(preset);
     return { ...state, preset, view: { tabs: [] }, payload: {} };
+  })
+  .on(API.LOCK_UI, (state, lockedUI) => {
+    return { ...state, lockedUI };
+  })
+  .on(API.HISTORY_UPDATE, (state, { undo, redo }) => {
+    return { ...state, undo, redo };
   });
 
 export default reducer;
\ No newline at end of file
diff --git a/Engine/Plugins/VirtualProduction/RemoteControlWebInterface/WebApp/Client/src/shared/index.ts b/Engine/Plugins/VirtualProduction/RemoteControlWebInterface/WebApp/Client/src/shared/index.ts
--- a/Engine/Plugins/VirtualProduction/RemoteControlWebInterface/WebApp/Client/src/shared/index.ts
+++ b/Engine/Plugins/VirtualProduction/RemoteControlWebInterface/WebApp/Client/src/shared/index.ts
@@ -1,25 +1,31 @@
 export enum PropertyType {
-  Boolean =     'bool',
-  Int8 =        'int8',
-  Int16 =       'int16',
-  Int32 =       'int32',
-  Int64 =       'int64',
-  Uint8 =       'uint8',
-  Uint16 =      'uint16',
-  Uint32 =      'uint32',
-  Uint64 =      'uint64',
-  Float =       'float',
-  Double =      'double',
-  Vector =      'FVector',
-  Vector2D =    'FVector2D',
-  Vector4 =     'FVector4',
-  Rotator =     'FRotator',
-  Color =       'FColor',
-  LinearColor = 'FLinearColor',
-  String =      'FString',
-  Text =        'FText',
-
-  Function =    'Function',
+  Boolean =         'bool',
+  Int8 =            'int8',
+  Int16 =           'int16',
+  Int32 =           'int32',
+  Int64 =           'int64',
+  Uint8 =           'uint8',
+  Uint16 =          'uint16',
+  Uint32 =          'uint32',
+  Uint64 =          'uint64',
+  Float =           'float',
+  Double =          'double',
+  Vector =          'FVector',
+  Vector2D =        'FVector2D',
+  Vector4 =         'FVector4',
+  Rotator =         'FRotator',
+  Color =           'FColor',
+  LinearColor =     'FLinearColor',
+  String =          'FString',
+  Text =            'FText',
+
+  Function =        'Function',
+}
+
+export enum ConnectionSignal {
+  Good =   'GOOD',
+  Normal = 'NORMAL',
+  Bad =    'BAD',
 }
 
 export interface ColorProperty {
@@ -33,6 +39,7 @@ export interface VectorProperty {
   X: number;
   Y: number;
   Z: number;
+  W?: number;
 }
 
 export interface RotatorProperty {
@@ -121,6 +128,7 @@ export interface IPreset {
   ExposedProperties?: IExposedProperty[];
   ExposedFunctions?: IExposedFunction[];
   Exposed?: Record<string, IExposedProperty | IExposedFunction>;
+  IsFavorite?: boolean;
 }
 
 export type IPresets = { [preset: string]: IPreset };
@@ -140,6 +148,7 @@ export type IPayload = { [property: string]: PropertyValue | IPayload };
 export type IPayloads = { [preset: string]: IPayload };
 
 export enum WidgetTypes {
+  Asset =           'Asset',
   Dial =            'Dial',
   Dials =           'Dials',
   Slider =          'Slider',
@@ -193,11 +202,16 @@ export interface IColorPickerList {
 export enum TabLayout {
   Stack =    'Stack',
   Screen =   'Screen',
+  Empty =    'Empty',
 }
 
 export enum ScreenType {
-  Snapshot =      'Snapshot',
-  Sequencer =     'Sequencer',
+  Stack =           'Stack',
+  Playlist =        'Playlist',
+  Snapshot =        'Snapshot',
+  Sequencer =       'Sequencer',
+  ColorCorrection = 'ColorCorrection',
+  LightCards =      'LightCards',
 }
 
 export interface IScreen {
@@ -263,4 +277,10 @@ export interface ITab {
   layout: TabLayout;
   panels?: IPanel[];
   screen?: IScreen;
+}
+
+export interface IHistory {
+  property: string;
+  value: any;
+  time: Date;
 }
\ No newline at end of file
diff --git a/Engine/Plugins/VirtualProduction/RemoteControlWebInterface/WebApp/Client/src/utilities/WidgetUtilities.ts b/Engine/Plugins/VirtualProduction/RemoteControlWebInterface/WebApp/Client/src/utilities/WidgetUtilities.ts
--- a/Engine/Plugins/VirtualProduction/RemoteControlWebInterface/WebApp/Client/src/utilities/WidgetUtilities.ts
+++ b/Engine/Plugins/VirtualProduction/RemoteControlWebInterface/WebApp/Client/src/utilities/WidgetUtilities.ts
@@ -158,7 +158,8 @@ export class WidgetUtilities {
     if (propertyType && (propertyType.startsWith('TEnum') || propertyType.startsWith('E')))
       return [WidgetTypes.Dropdown];
 
-    return [];
+    // Unknown Property, show as a Asset
+    return [WidgetTypes.Asset];
   }
 
   static colorToRgb(color: ColorProperty | VectorProperty, max: number) {
diff --git a/Engine/Plugins/VirtualProduction/RemoteControlWebInterface/WebApp/Server/src/Api.ts b/Engine/Plugins/VirtualProduction/RemoteControlWebInterface/WebApp/Server/src/Api.ts
--- a/Engine/Plugins/VirtualProduction/RemoteControlWebInterface/WebApp/Server/src/Api.ts
+++ b/Engine/Plugins/VirtualProduction/RemoteControlWebInterface/WebApp/Server/src/Api.ts
@@ -54,6 +54,11 @@ export class Api {
     this.send(res, UnrealEngine.getView(req.query.preset?.toString()));
   }
 
+  @Put('presets/:preset/favorite')
+  private favorite(req: Request, res: Response) {
+    this.send(res, UnrealEngine.favoritePreset(req.params.preset, !!req.body?.value));
+  }
+
   @Put('proxy')
   private proxy(req: Request, res: Response) {
     this.send(res, UnrealEngine.proxy(req.body?.method, req.body?.url, req.body?.body));
diff --git a/Engine/Plugins/VirtualProduction/RemoteControlWebInterface/WebApp/Server/src/Notify.ts b/Engine/Plugins/VirtualProduction/RemoteControlWebInterface/WebApp/Server/src/Notify.ts
--- a/Engine/Plugins/VirtualProduction/RemoteControlWebInterface/WebApp/Server/src/Notify.ts
+++ b/Engine/Plugins/VirtualProduction/RemoteControlWebInterface/WebApp/Server/src/Notify.ts
@@ -12,6 +12,7 @@ export namespace Notify {
 
     io.sockets.on('connection', (socket: socketio.Socket) => {
       socket
+        .on('client', onPing)
         .on('view', onViewChange)
         .on('value', UnrealEngine.setPayloadValue)
         .on('execute', UnrealEngine.executeFunction)
@@ -50,4 +51,8 @@ export namespace Notify {
   export function emitPassphraseChanged(wrongPassphrase: string) {
     io.emit('passphrase', wrongPassphrase);
   }
+
+  export function onPing(time: number) {
+    io.emit('pong', time, new Date().getTime());
+  }
 }
\ No newline at end of file
diff --git a/Engine/Plugins/VirtualProduction/RemoteControlWebInterface/WebApp/Server/src/Program.ts b/Engine/Plugins/VirtualProduction/RemoteControlWebInterface/WebApp/Server/src/Program.ts
--- a/Engine/Plugins/VirtualProduction/RemoteControlWebInterface/WebApp/Server/src/Program.ts
+++ b/Engine/Plugins/VirtualProduction/RemoteControlWebInterface/WebApp/Server/src/Program.ts
@@ -24,7 +24,6 @@ class Program {
         case '--dev':
           this.dev = true;
           this.port = 7001;
-          this.logger = true;
           break;
 
         case '--monitor':
diff --git a/Engine/Plugins/VirtualProduction/RemoteControlWebInterface/WebApp/Server/src/UnrealEngine.ts b/Engine/Plugins/VirtualProduction/RemoteControlWebInterface/WebApp/Server/src/UnrealEngine.ts
--- a/Engine/Plugins/VirtualProduction/RemoteControlWebInterface/WebApp/Server/src/UnrealEngine.ts
+++ b/Engine/Plugins/VirtualProduction/RemoteControlWebInterface/WebApp/Server/src/UnrealEngine.ts
@@ -3,6 +3,7 @@ import { IPayload, IPayloads, IPreset, IPresets, IPanel, IView, ICustomStackWidg
 import _ from 'lodash';
 import WebSocket from 'ws';
 import { Notify, Program, LogServer } from './';
+import fs from 'fs-extra';
 import crypto from 'crypto';
 
 
@@ -81,15 +82,23 @@ export namespace UnrealEngine {
   let isOpen: boolean;
 
   let pendings: { [id: string]: (reply: any) => void } = {};
-  let request: number = 1;
+  let httpRequest: number = 1;
+  let wsRequest: number = 1;
 
   let presets: IPresets = {};
   let registered: { [id: string]: boolean };
   let payloads: IPayloads = {};
   let views: { [preset: string]: IView } = {};
   let workingPassphrases: string[] = [];
+  let favorites: { [preset: string]: boolean } = {};
 
   export async function initialize() {
+    try {
+      if (await fs.pathExists('./favorites.json'))
+        favorites = await fs.readJSON('./favorites.json');
+    } catch {
+    }
+
     connect();
     startQuitTimeout(60 * 1000);
   }
@@ -115,7 +124,7 @@ export namespace UnrealEngine {
     if (connection?.readyState === WebSocket.OPEN || connection?.readyState === WebSocket.CONNECTING)
       return;
 
-    const address = `ws://localhost:${Program.ueWebSocketPort}`;
+    const address = `ws://127.0.0.1:${Program.ueWebSocketPort}`;
 
     connection?.removeAllListeners();
 
@@ -129,10 +138,10 @@ export namespace UnrealEngine {
 
   export async function checkPassphrase(passphrase: string) : Promise<boolean> {
     try {
-      let res = await get<UnrealApi.PassphraseCheck>(`/remote/passphrase/`, passphrase);
-      if (!res?.keyCorrect) {
+      const res = await get<UnrealApi.PassphraseCheck>(`/remote/passphrase/`, passphrase);
+      if (!res?.keyCorrect)
         return false;
-      }
+
     } catch (error) {
       return false;
     }
@@ -141,7 +150,7 @@ export namespace UnrealEngine {
       workingPassphrases.push(passphrase);
       
       if (workingPassphrases.length === 1 || isOpen) {
-        await pullPresets(true);
+        await pullPresets();
         pullTimer();   
       }
     }
@@ -161,6 +170,7 @@ export namespace UnrealEngine {
     registered = {};
     payloads = {};
     views = {};
+    favorites = {};
     isOpen = undefined;
     
     clearQuitTimeout();
@@ -168,15 +178,16 @@ export namespace UnrealEngine {
     console.log('Connected to UE Remote WebSocket');
     Notify.emit('connected', true);
     
-    isOpen = await checkPassphrase("");
+    isOpen = await checkPassphrase('');
+
     Notify.emit('opened', isOpen);
     if (isOpen || workingPassphrases.length > 0) {
-      await pullPresets(true);
+      await refresh();
       pullTimer();
     }
   }
 
-  async function refresh() { 
+  async function refresh() {
     try {
       await pullPresets();
     } catch (error) {
@@ -189,9 +200,8 @@ export namespace UnrealEngine {
 
       // First check if there are any new presets
       for (const preset of Presets) {
-        if (preset && !presets[preset.ID]){
+        if (preset && !presets[preset.ID])
           await refreshPreset(preset.ID, preset.Name);
-        }
       }
 
       // Check for deleted presets
@@ -201,16 +211,15 @@ export namespace UnrealEngine {
           delete presets[id];
           delete payloads[id];
         }
-        
-        Notify.emit('presets', Object.values(presets));
+
+        Notify.emit('presets', getSortedPresets());
       }
     } catch (error) {
     }
 
     // Check for new presets once every 5 seconds
-    if (isConnected()) {
-      autoPullTimeout = setTimeout(pullTimer, 5 * 1000);
-    }
+    if (isConnected())
+       autoPullTimeout = setTimeout(pullTimer, 5 * 1000);
   }
  
   async function onMessage(data: WebSocket.Data) {
@@ -218,14 +227,13 @@ export namespace UnrealEngine {
 
     try {
       const message = JSON.parse(json);
-      
       if (message.RequestId) {
         const promise = pendings[message.RequestId];
         if (!promise)
-        return;
-        
+          return;
+
         if (Program.logger)
-        LogServer.logLoopback({ RequestId: message.RequestId, Stage: 'WebApp Done' }, true);
+          LogServer.logLoopback({ RequestId: message.RequestId, Stage: 'WebApp Done' }, true);
         
         if (message.Verb === "401" && workingPassphrases.includes(message.Passphrase)) {
           Notify.emitPassphraseChanged(message.Passphrase);
@@ -241,10 +249,10 @@ export namespace UnrealEngine {
         return;
       }
 
-      if (isOpen !== message.keyCorrect) {
-        isOpen = !isOpen;
+      if (message.keyCorrect !== undefined && isOpen !== message.keyCorrect) {
+        isOpen = message.keyCorrect;
         Notify.emit('opened', isOpen);
-        await pullPresets(true);
+        await pullPresets();
         pullTimer();
       }
 
@@ -269,7 +277,7 @@ export namespace UnrealEngine {
         }
 
         case UnrealApi.PresetEvent.FieldsAdded: {
-          const preset = populatePreset(message.Description as IPreset);
+          const preset = await populatePreset(message.Description as IPreset);
           const values = await pullPresetValues(preset);
           for (const property in values) {
             setPayloadValueInternal(payloads, [preset.ID, property], values[property]);
@@ -295,15 +303,15 @@ export namespace UnrealEngine {
             break;
 
           await refreshPreset(message.PresetId, message.PresetName);
-          await refreshView(message.PresetId, message.Metadata.view);
+          refreshView(message.PresetId, message.Metadata.view);
           break;
         }
 
         case UnrealApi.PresetEvent.LayoutModified: {
-          const preset = populatePreset(message.Preset as IPreset);
+          const preset = await populatePreset(message.Preset as IPreset);
           if (preset?.ID) {
             presets[preset.ID] = preset;
-            Notify.emit('presets', Object.values(presets));
+            Notify.emit('presets', getSortedPresets());
           }
 
           break;
@@ -324,6 +332,7 @@ export namespace UnrealEngine {
     registered = {};
     payloads = {};
     views = {};
+    favorites = {};
     isOpen = undefined;
 
     Notify.emit('connected', false);
@@ -358,13 +367,13 @@ export namespace UnrealEngine {
 
   function send(message: string, parameters: any, passphrase?: string) {
     verifyConnection();
-    const Id = parameters?.RequestId ?? request++;
+    const Id = parameters?.RequestId ?? wsRequest++;
     passphrase = passphrase ?? workingPassphrases[0];
     connection.send(JSON.stringify({ MessageName: message, Id, Passphrase: passphrase, Parameters: parameters}));
   }
 
   function http<T>(Verb: string, URL: string, Body?: object, wantAnswer?: boolean, passphrase?: string): Promise<T> {
-    const RequestId = request++;
+    const RequestId = httpRequest++;
     const payload = { RequestId, Verb, URL, Body };
     if (Program.logger)
       LogServer.logLoopback(payload);
@@ -392,55 +401,44 @@ export namespace UnrealEngine {
 
   function unregisterPreset(PresetName: string) {
     send('preset.unregister', { PresetName });
-  }    
+  }
+
+  function getSortedPresets() {
+    return _.orderBy(Object.values(presets), ['IsFavorite', 'Name'], ['desc', 'asc']);
+  }
 
   export async function getPresets(): Promise<IPreset[]> {
-    return Object.values(presets);
+    return getSortedPresets();
   }
 
-  async function pullPresets(bInitial?: boolean): Promise<void> {
+  async function pullPresets(): Promise<void> {
     if (isPullingPresets)
       return;
 
     isPullingPresets = true;
-    if (bInitial)
-      setLoading(true);
+    setLoading(true);
 
-    try {
-      const allPresets: IPresets = {};
-      let allPayloads: IPayloads = {};
-  
+    try { 
       const { Presets } = await get<UnrealApi.Presets>('/remote/presets');
-      
-      for (const p of Presets ?? []) {
-        const Preset = await pullPreset(p.ID, p.Name);
-        if (!Preset)
-          continue;
 
-        allPresets[Preset.ID] = Preset;
-        if (bInitial || !presets[Preset.ID]){
-          allPayloads[Preset.ID] = await pullPresetValues(Preset);
-        }
+      for (const preset of Presets ?? []) {
+        const stub = preset as IPreset;
+        stub.Exposed = {};
+        stub.ExposedFunctions = [];
+        stub.ExposedProperties = [];
+        stub.Groups = [];
+        stub.IsFavorite = !!favorites?.[preset.ID];
+        presets[preset.ID] = stub;
       }
 
-      if (!equal(presets, allPresets)) {
-        presets = allPresets;
-        Notify.emit('presets', Object.values(presets));
-      }
-
-      allPayloads = {...payloads, ...allPayloads };
-      if (bInitial || !equal(payloads, allPayloads)) {
-        payloads = allPayloads;
-        Notify.emit('payloads', payloads);
-      }
+      Notify.emit('presets', getSortedPresets());
 
     } catch (error) {
       console.log('Failed to pull presets data');
     }
 
     isPullingPresets = false;
-    if (bInitial)
-      setLoading(false);
+    setLoading(false);
   }
 
   async function pullPreset(id: string, name?: string): Promise<IPreset> {
@@ -454,7 +452,7 @@ export namespace UnrealEngine {
       if (!Preset)
         return null;
 
-      populatePreset(Preset);
+      await populatePreset(Preset);
       return Preset;
     } catch (error) {
       console.log(`Failed to pull preset '${name || id}' data`);
@@ -467,11 +465,11 @@ export namespace UnrealEngine {
       return;
 
     presets[id] = preset;
-    Notify.emit('presets', Object.values(presets));
+    Notify.emit('presets', getSortedPresets());
     return preset;
   }
 
-  function populatePreset(preset: IPreset) {
+  async function populatePreset(preset: IPreset) {
     preset.ExposedProperties = [];
     preset.ExposedFunctions = [];
     preset.Exposed = {};
@@ -495,17 +493,19 @@ export namespace UnrealEngine {
       preset.ExposedFunctions.push(...Group.ExposedFunctions);
     }
 
+    preset.IsFavorite = favorites?.[preset.ID];
+
     return preset;
   }
 
-  async function refreshView(id: string, viewJson: string) {
+  function refreshView(id: string, viewJson: string) {
     try {
       if (!viewJson)
         return;
 
       const view = JSON.parse(viewJson) as IView;
       if (equal(view, views[id]))
-        return; 
+        return;
 
       for (const tab of view.tabs) {
         if (!tab.panels)
@@ -519,7 +519,7 @@ export namespace UnrealEngine {
       Notify.onViewChange(id, view, true);
     } catch (error) {
       console.log('Failed to parse View of Preset', id);
-    }    
+    }
   }
 
   function setPanelIds(panel: IPanel) {
@@ -609,7 +609,7 @@ export namespace UnrealEngine {
     const Requests = [];
     for (const property of Preset.ExposedProperties) {
       Requests.push({
-        RequestId: request++,
+        RequestId: httpRequest++,
         Verb: 'GET',
         URL: `/remote/preset/${Preset.ID}/property/${property.ID}`,
         PropertyId: property.ID,
@@ -645,7 +645,7 @@ export namespace UnrealEngine {
     Notify.emit('payloads', payloads);
     
     const res = await get<UnrealApi.View>(`/remote/preset/${id}/metadata/view`);
-    await refreshView(id, res?.Value);
+    refreshView(id, res?.Value);
   }
 
   export async function getPayload(preset: string): Promise<IPayload> {
@@ -674,7 +674,7 @@ export namespace UnrealEngine {
       if (Program.logger) {
         const presetObj = presets[preset];
         const log = {
-          RequestId: request,
+          RequestId: httpRequest,
           Stage: 'Set Property',
           Preset: presetObj?.Name,
           Property: presetObj?.Exposed[property]?.DisplayName,
@@ -732,7 +732,7 @@ export namespace UnrealEngine {
       const prop = presets[preset]?.Exposed[property];
       if (prop) {
         prop.Metadata[metadata] = value;
-        Notify.emit('presets', Object.values(presets));
+        Notify.emit('presets', getSortedPresets());
       }
     } catch (err) {
       console.log('Failed to set property metadata');
@@ -766,6 +766,21 @@ export namespace UnrealEngine {
     Notify.emitValuesChanges(preset, changes);
   }
 
+  export async function favoritePreset(preset: string, value: boolean): Promise<void> {
+    if (!presets[preset])
+      return;
+
+    presets[preset].IsFavorite = value;
+
+    if (value)
+      favorites[preset] = true;
+    else
+      delete favorites[preset];
+
+    await fs.writeJson('./favorites.json', favorites, { spaces: 2 });
+    Notify.emit('presets', getSortedPresets());
+  }
+
   export async function getView(preset: string): Promise<IView> {
     return views[preset];
   }
diff --git a/Engine/Plugins/VirtualProduction/RemoteControlWebInterface/WebApp/Start.bat b/Engine/Plugins/VirtualProduction/RemoteControlWebInterface/WebApp/Start.bat
--- a/Engine/Plugins/VirtualProduction/RemoteControlWebInterface/WebApp/Start.bat
+++ b/Engine/Plugins/VirtualProduction/RemoteControlWebInterface/WebApp/Start.bat
@@ -1,8 +1,9 @@
 @echo off
 
-rem Add default nodejs installation folder to path, in case it was not added / overwritten
+REM Add default nodejs installation folder to path, in case it was not added / overwritten
 SET PATH=%PATH%;%ProgramFiles%\nodejs\
 
+
 rem Check if nodejs is in the env variable PATH 
 for %%X in (node.exe) do (set node=%%~$PATH:X)
 if not defined node (
@@ -19,5 +20,10 @@ if %errorlevel% LSS 14 (
   exit 1
 )
 
+if %errorlevel% GEQ 17 (
+  REM Due to changes on Node.js v17, --openssl-legacy-provider was added for handling key size on OpenSSL v3
+  SET NODE_OPTIONS=--openssl-legacy-provider
+)
+
 rem redirecting all command line arguments to node script
 node Scripts/start.js %*
\ No newline at end of file
diff --git a/Engine/Plugins/VirtualProduction/RemoteControlWebInterface/WebApp/Start.sh b/Engine/Plugins/VirtualProduction/RemoteControlWebInterface/WebApp/Start.sh
--- a/Engine/Plugins/VirtualProduction/RemoteControlWebInterface/WebApp/Start.sh
+++ b/Engine/Plugins/VirtualProduction/RemoteControlWebInterface/WebApp/Start.sh
@@ -23,6 +23,12 @@ if [ ${VERSION} -lt 14 ] ; then
   exit 1
 fi
 
+if [ ${VERSION} -ge 17 ] ; then
+  #Due to changes on Node.js v17, --openssl-legacy-provider was added for handling key size on OpenSSL v3
+  export NODE_OPTIONS=--openssl-legacy-provider
+fi
+
+
 FOLDER=$(dirname "$0")
 node ${FOLDER}/Scripts/start.js "$@"
 
