RemoteControlWebInterface: Auto download nodejs and use a dependecy (Windows only)
#jria UE-161346

diff --git a/Engine/Plugins/VirtualProduction/RemoteControlWebInterface/WebApp/Clean.bat b/Engine/Plugins/VirtualProduction/RemoteControlWebInterface/WebApp/Clean.bat
--- a/Engine/Plugins/VirtualProduction/RemoteControlWebInterface/WebApp/Clean.bat
+++ b/Engine/Plugins/VirtualProduction/RemoteControlWebInterface/WebApp/Clean.bat
@@ -16,7 +16,7 @@ goto :EOF
 :DEL_FOLDERS
 REM Using Robocopy instead of rmdir
 REM as node_modules may contains really long file path, and rmdir fails to delete them
-echo Deleting %~1
+REM echo Deleting %~1
 if not exist empty\ mkdir empty
 robocopy empty "%~1" /mir > nul
 rmdir /q /s "%~1"
@@ -29,6 +29,6 @@ REM Files are marked as readonly because of perforce
 REM Node v16 fails to install dependencies because it tries to overwrite the file
 REM if we force node to skip it, earlier versions of node will fail if we do that
 REM so we have original files prefixed with _, at build they are being copied and the readonly attribute is removed
-del /f /q "%~1\package.json"
-del /f /q "%~1\package-lock.json"
+if exist "%~1\package.json"         del /f /q "%~1\package.json"
+if exist "%~1\package-lock.json"    del /f /q "%~1\package-lock.json"
 goto :EOF
\ No newline at end of file
diff --git a/Engine/Plugins/VirtualProduction/RemoteControlWebInterface/WebApp/Start.bat b/Engine/Plugins/VirtualProduction/RemoteControlWebInterface/WebApp/Start.bat
--- a/Engine/Plugins/VirtualProduction/RemoteControlWebInterface/WebApp/Start.bat
+++ b/Engine/Plugins/VirtualProduction/RemoteControlWebInterface/WebApp/Start.bat
@@ -1,19 +1,58 @@
 @echo off
+setlocal enabledelayedexpansion
+setlocal enableextensions
 
-REM Add default nodejs installation folder to path, in case it was not added / overwritten
+pushd "%~dp0"
+
+SET NodeVersion=v16.17.0
+SET DownloadedNodeFolder=Node-%NodeVersion%
+SET NodeName=node-%NodeVersion%-win-x64
+
+if not exist "%DownloadedNodeFolder%\" (
+  echo Downloading NodeJS for Windows...
+  
+  @REM Download nodejs and follow redirects.
+  curl -s -L -o .\%NodeName%.zip "https://nodejs.org/dist/%NodeVersion%/%NodeName%.zip"
+
+  @REM Only if download succssed
+  if exist %NodeName%.zip (
+    @REM cleanup node_modules, as they might be of an different nodejs bersion
+    CALL Clean.bat
+
+    @REM Unarchive the .zip
+    tar -xf %NodeName%.zip
+
+    @REM Rename the extracted, versioned, directory that contains the NodeJS binaries to simply "Node".
+    ren "%NodeName%\" "%DownloadedNodeFolder%"
+
+    @REM Delete the downloaded node.zip
+    del %NodeName%.zip
+  ) else (
+    echo Failed to download NodeJS %NodeVersion%
+  )
+)
+
+if exist "%DownloadedNodeFolder%\" (
+  @REM Add downloaded nodejs version to be first in line
+  echo Using downloaded NodeJS version from folder %cd%\%DownloadedNodeFolder%
+  SET "PATH=%cd%\%DownloadedNodeFolder%;%PATH%;"
+)
+
+@REM Add default nodejs installation folder to path, in case it was not added / overwritten
 SET PATH=%PATH%;%ProgramFiles%\nodejs\
 
 
-rem Check if nodejs is in the env variable PATH 
+@REM Check if nodejs is in the env variable PATH 
 for %%X in (node.exe) do (set node=%%~$PATH:X)
 if not defined node (
   echo ERROR: Couldn't find node.js installed, Please install latest nodejs from https://nodejs.org/en/download/
   exit 1
 )
 
-rem Let's check if it is a modern nodejs
+
+@REM Let's check if it is a modern nodejs
 node -e "process.exit( process.versions.node.split('.')[0] );"
-echo Found Node.js version %errorlevel% (%node%)
+echo Using Node.js version %errorlevel%
 
 if %errorlevel% LSS 14 (
   echo ERROR: installed node.js version is too old, please install latest nodejs from https://nodejs.org/en/download/
@@ -21,9 +60,9 @@ if %errorlevel% LSS 14 (
 )
 
 if %errorlevel% GEQ 17 (
-  REM Due to changes on Node.js v17, --openssl-legacy-provider was added for handling key size on OpenSSL v3
+  @REM Due to changes on Node.js v17, --openssl-legacy-provider was added for handling key size on OpenSSL v3
   SET NODE_OPTIONS=--openssl-legacy-provider
 )
 
-rem redirecting all command line arguments to node script
-node Scripts/start.js %*
\ No newline at end of file
+@REM redirecting all command line arguments to node script
+node Scripts/start.js %*
