RemoteControlWebInterface: Merge from Marmot427 (Double click To Enable/Disable)
Fix unlinked properties in new LightCard actor

diff --git a/Engine/Plugins/VirtualProduction/RemoteControlWebInterface/WebApp/Client/src/components/Screens/ColorCorrection/ColorCorrection.tsx b/Engine/Plugins/VirtualProduction/RemoteControlWebInterface/WebApp/Client/src/components/Screens/ColorCorrection/ColorCorrection.tsx
--- a/Engine/Plugins/VirtualProduction/RemoteControlWebInterface/WebApp/Client/src/components/Screens/ColorCorrection/ColorCorrection.tsx
+++ b/Engine/Plugins/VirtualProduction/RemoteControlWebInterface/WebApp/Client/src/components/Screens/ColorCorrection/ColorCorrection.tsx
@@ -141,6 +141,20 @@ export class ColorCorrection extends React.Component<Props, State> {
     this.setState({ selected, value }, () => this.onRefreshValues(true));
   }
 
+  onToggleActorEnable = (selected: string) => {
+    const { values, section } = this.state;
+    const value = values[selected] ?? {};
+    switch (section) {
+      case Section.ColorCorrection:
+        this.onPropertyChange('Enabled', !value.Enabled);
+        break;
+    
+      case Section.LightCards:
+        this.onPropertyChange('bHidden', !value.bHidden);
+        break;
+    }
+  }
+
   onSpawnActor = async () => {
     const { section } = this.state;
     this.updateSelection = true;
@@ -225,15 +239,20 @@ export class ColorCorrection extends React.Component<Props, State> {
   }
 
   getRootComponentPath = () => {
-    const { selected } = this.state;
+    const { section, selected } = this.state;
     if (!selected)
       return selected;
 
+    let root = '.Root';
+    if (section === Section.LightCards)
+      root = '.DefaultSceneRoot';
+    
+
     let index = selected.lastIndexOf("'");
     if (index < 0)
-      return `${selected}.DefaultSceneRoot`;
+      return `${selected}${root}`;
 
-    return selected.slice(0, index) + '.DefaultSceneRoot' + selected.slice(index);
+    return selected.slice(0, index) + root + selected.slice(index);
   }
 
   onRefreshValues = async (update = false) => {
@@ -324,8 +343,15 @@ export class ColorCorrection extends React.Component<Props, State> {
   onSectionChange = (section: Section) => {
     let { mode } = this.state;
 
-    if (section === Section.ColorCorrection && mode === CCRMode.Color)
-      mode = CCRMode.Global;
+    switch (section) {
+      case Section.ColorCorrection:
+        mode = CCRMode.Global;
+        break;
+
+      case Section.LightCards:
+        mode = CCRMode.Color;
+        break;
+    }
 
     this.setState({ section, mode, actors: {}, selected: undefined, root: undefined }, this.onRefresh);
   }
@@ -362,21 +388,10 @@ export class ColorCorrection extends React.Component<Props, State> {
   getColorProperty = () => {
     const { section, mode, colorProperty } = this.state;
 
-    if (mode === CCRMode.Color)
+    if (section === Section.LightCards)
       return 'Color';
 
-    let root = '';
-    switch (section) {
-      case Section.ColorCorrection:
-        root = 'ColorGradingSettings';
-        break;
-
-      case Section.LightCards:
-        root = 'color grading';
-        break;
-    }
-
-    return `${root}.${mode}.${colorProperty}`;
+    return `ColorGradingSettings.${mode}.${colorProperty}`;
   }
 
   renderColorSlider = (property: string, min: number, max: number, reset: number, type: string) => {
@@ -449,7 +464,7 @@ export class ColorCorrection extends React.Component<Props, State> {
 
     let opt = [];
     if (Array.isArray(options)) {
-      opt = options.map(value => ({ value, label: value }));
+      opt = options.map(value => ({ value, label: _.startCase(value) }));
     } else {
       for (const value in options)
         opt.push({ value, label: options[value] });
@@ -545,8 +560,33 @@ export class ColorCorrection extends React.Component<Props, State> {
     </>;
   }
 
+  getActorClassName = (actor: string) => {
+    const { selected, section, values } = this.state;
+
+    let className = 'light-card ';
+    if (selected === actor)
+      className += 'active ';
+
+    let disabled;
+    const value = values[actor];
+    switch (section) {
+      case Section.ColorCorrection:
+        disabled = (value?.Enabled === false);
+        break;
+    
+      case Section.LightCards:
+        disabled = (value?.bHidden === true);
+        break;
+    }
+
+    if (disabled === true)
+      className += 'disabled ';
+
+    return className;
+  }
+
   renderActorsContent = (pagination: number, items: number, actors: Record<string, any>, favorite: boolean) => {
-    const { selected, search, favorites, section, values } = this.state;
+    const { search, favorites, section, values } = this.state;
 
     let start = pagination * items;
     let end = start + items;
@@ -562,32 +602,22 @@ export class ColorCorrection extends React.Component<Props, State> {
       <>
         {actorsKeys.slice(start, end).map(actor =>
           <div key={actor}
-               className={`light-card ${selected === actor ? 'active' : ''}`}
-               onClick={() => this.onActorSelect(actor)}>
+               className={this.getActorClassName(actor)}
+               onClick={() => this.onActorSelect(actor)}
+               onDoubleClick={() => this.onToggleActorEnable(actor)}>
             {_.startCase(this.state.actors[actor])}
-            { section === Section.LightCards &&
-              <div className={`power-window-icon ${_.get(values, [actor, 'Power Window Mode']) ? 'active' : ''}`}>
-                <svg width="28" height="29" viewBox="0 0 28 29" fill="none" xmlns="http://www.w3.org/2000/svg">
-                  <path d="M21.7 6.77118V23.221L4.55 20.2461V9.74609L21.7 6.77118ZM24.5 4.49609L3.5 8.87109V21.1211L24.5 25.4961V4.49609Z" fill="white" fillOpacity="0.35"/>
-                  <path opacity="0.5" d="M19.075 20.2461L7 18.1462V11.4961L19.075 20.2461Z" fill="white" fillOpacity="0.35"/>
-                  <path fillRule="evenodd" clipRule="evenodd" d="M19.25 15.5219L12.075 10.4468L19.25 9.22168V15.5219Z" fill="white" fillOpacity="0.35"/>
-                  <path opacity="0.5" d="M1.75 2.74609H7V0.996094H0V7.99609H1.75V2.74609Z" fill="white" fillOpacity="0.35"/>
-                  <path opacity="0.5" d="M26.25 2.74609V7.99609H28V0.996094H21V2.74609H26.25Z" fill="white" fillOpacity="0.35"/>
-                  <path opacity="0.5" d="M26.25 27.2461H21V28.9961H28V21.9961H26.25V27.2461Z" fill="white" fillOpacity="0.35"/>
-                  <path opacity="0.5" d="M1.75 27.2461V21.9961H0V28.9961H7V27.2461H1.75Z" fill="white" fillOpacity="0.35"/>
-                </svg>
-              </div>
-            }
             {!!favorites[actor] && 
               <FontAwesomeIcon className="icon-favorite" icon={['fas', 'star']} />
             }
 
-            <div className='change-indicators'>
-              {this.renderChangeIndication(actor, CCRMode.Global)}
-              {this.renderChangeIndication(actor, CCRMode.Highlights)}
-              {this.renderChangeIndication(actor, CCRMode.Midtones)}
-              {this.renderChangeIndication(actor, CCRMode.Shadows)}
-            </div>
+            { section === Section.ColorCorrection &&
+              <div className='change-indicators'>
+                {this.renderChangeIndication(actor, CCRMode.Global)}
+                {this.renderChangeIndication(actor, CCRMode.Highlights)}
+                {this.renderChangeIndication(actor, CCRMode.Midtones)}
+                {this.renderChangeIndication(actor, CCRMode.Shadows)}
+              </div>
+            }
           </div>
         )}
         {!favorite && (
@@ -608,13 +638,15 @@ export class ColorCorrection extends React.Component<Props, State> {
         this.renderSlider('Temperature', 0, 10_000, 6_500),
         this.renderSlider('Tint', -1, 1, 0, 3),
         this.renderSlider('Exposure', -2, 10, 0.5, 3, 'Exposure'),
-        this.renderSlider('Intensity', -2, 10, 1, 3, null, _.get(this.state.value, 'Power Window Mode', false)),
+        this.renderSlider('Gain', -2, 10, 1, 3),
       );
 
     if (pagination === 1)
       contents.push(
         this.renderSlider('Opacity', 0, 1, 1, 3),
         this.renderSlider('Feathering', 0, 1, 0, 3),
+        <div />,
+        <div />,
       );
 
     return contents;
@@ -636,26 +668,8 @@ export class ColorCorrection extends React.Component<Props, State> {
       contents.push(
         this.renderSlider('Scale.X', 0, 20, 2, 3, 'Scale X'),
         this.renderSlider('Scale.Y', 0, 20, 2, 3, 'Scale Y'),
-        this.renderSlider('WallIndex', -1, 20, -1, 0),
-      );
-
-    if (pagination === 2)
-      contents.push(
-        this.renderSlider('Unit', 0, 200, 0),
-        <div className="buttons-layout" key="Updown">
-          <div className="buttons-content">
-            {this.renderButton('Move Up')}
-            {this.renderButton('Move Down')}
-          </div>
-          <div className="buttons-content small-button"  key="Sides">
-            {this.renderButton('Move X+')}
-            <div className="space-buttons">
-              {this.renderButton('Move Y-')}
-              {this.renderButton('Move Y+')}
-            </div>
-            {this.renderButton('Move X-')}
-          </div>
-        </div>
+        <div />,
+        <div />,
       );
 
     return contents;
@@ -740,21 +754,17 @@ export class ColorCorrection extends React.Component<Props, State> {
             <TabPane id="appearance" tab="Appearance">
               <div className="grid-row">
                 <div className="props-row">
-                  {this.renderSelect('Mask', { NewEnumerator1: 'Square', NewEnumerator0: 'Circle' })}
+                  {this.renderSelect('Mask', ['Square', 'Circle', 'Polygon', 'UseTextureAlpha'])}
                   <div className="prop asset-control">
                     <AssetWidget browse
                                  reset={false}
-                                 type="Texture2D"
+                                 type="Texture"
                                  label="Texture"
                                  onChange={v => this.onPropertyChange('Texture', v)}
                                  value={_.get(value, 'Texture') ?? ''} />
                     <span className="reset"><FontAwesomeIcon icon={['fas', 'undo']} onClick={() => this.onPropertyChange('Texture', null)} /></span>
                   </div>
                 </div>
-                <div className="props-row">
-                  {this.renderToggle('Power Window Mode', 'Power Window')}
-                  {this.renderToggle('useCCmult', 'CC Multiple')}
-                </div>
               </div>
               <div className="multilines-content">
                 <PaginationContent pages={2} key="appearance">
@@ -763,23 +773,8 @@ export class ColorCorrection extends React.Component<Props, State> {
               </div>
             </TabPane>
             <TabPane id="orientation" tab="Orientation">
-              <div className="props-row grid-row orientation">
-                <div className="prop-item">
-                  {this.renderToggle('UseClockSystem', 'Clock System')}
-                  <div className="prop">
-                    <TimeControl label="Time"
-                                  value={_.get(value, 'Hour')}
-                                  disabled={!_.get(value, 'UseClockSystem')}
-                                  onChange={v => this.onPropertyChange('Hour', v)} />
-                  </div>
-                </div>
-                <div className="prop-item">
-                  {this.renderToggle('StickToWall')}
-                  {this.renderButton('Set To Center of Stage')}
-                </div>
-              </div>
               <div className="multilines-content">
-                <PaginationContent pages={3} key="orientation">
+                <PaginationContent pages={2} key="orientation">
                   {({ pagination }) => this.renderOrientationContent(pagination)}
                 </PaginationContent>
               </div>
@@ -796,12 +791,6 @@ export class ColorCorrection extends React.Component<Props, State> {
             <TabPane id="position" tab="Position">
               {this.renderPositionTab()}
             </TabPane>
-            <TabPane id="advanced" tab="Advanced">
-              <div className='control-rows light-advanced'>
-                {this.renderSlider('color grading.ShadowsMax', -1, 1, 0.09)}
-                {this.renderSlider('color grading.HighlightsMin', -1, 1, 0.5)}
-              </div>
-            </TabPane>
           </Tabs>
         );
     }
@@ -825,9 +814,8 @@ export class ColorCorrection extends React.Component<Props, State> {
     let availableFavorites = _.pickBy(favorites, (v, k) => !!actors[k]);
 
     return (
-      <div className="color-correction-wrapper custom-ui-wrapper" tabIndex={-1}>
+      <div className={`color-correction-wrapper custom-ui-wrapper ${section}`} tabIndex={-1}>
         <div className='ipad-debug'></div>
-        <img src='/images/ref001.jpg' className='ref'/>
         <div className="group">
           <CorrectionColorPicker section={section}
                                  property={property}
diff --git a/Engine/Plugins/VirtualProduction/RemoteControlWebInterface/WebApp/Client/src/components/Screens/ColorCorrection/CorrectionColorPicker.tsx b/Engine/Plugins/VirtualProduction/RemoteControlWebInterface/WebApp/Client/src/components/Screens/ColorCorrection/CorrectionColorPicker.tsx
--- a/Engine/Plugins/VirtualProduction/RemoteControlWebInterface/WebApp/Client/src/components/Screens/ColorCorrection/CorrectionColorPicker.tsx
+++ b/Engine/Plugins/VirtualProduction/RemoteControlWebInterface/WebApp/Client/src/components/Screens/ColorCorrection/CorrectionColorPicker.tsx
@@ -66,7 +66,7 @@ export class CorrectionColorPicker extends React.Component<Props, State> {
   }
 
   onHsvWheelMove = (sign: number, hue?: number) => {
-    let { range, color, sensitivity } = this.props;
+    let { range, color, mode, sensitivity } = this.props;
     let rgb = WidgetUtilities.colorToRgb(color, range.max);
 
     const hsv = WidgetUtilities.rgb2Hsv(rgb);
@@ -82,7 +82,7 @@ export class CorrectionColorPicker extends React.Component<Props, State> {
     }
 
     rgb = WidgetUtilities.hsv2rgb(hsv);
-    color = WidgetUtilities.rgbToColor(rgb, range.max);
+    color = WidgetUtilities.rgbToColor(rgb, range.max, mode !== CCRMode.Color);
 
     this.onColorChange(color);
   }
@@ -267,25 +267,26 @@ export class CorrectionColorPicker extends React.Component<Props, State> {
         </div>
         <div className="cc-group" style={style}>
           <div className="body" style={{ width: "calc(100% - 20px)" }}>
-            <div className="mode-selectors">
-              <Tabs onTabChange={this.props.onColorModeChange}
-                    onlyHeader
-                    defaultActiveKey={mode}>
-                {modes.map(mode => <TabPane key={mode}
-                                            id={mode} 
-                                            tab={mode}
-                                            view={() => this.renderModifyIndicator(mode)} />
-                )}
-
-              </Tabs>
-            </div>
+            {section === Section.ColorCorrection &&
+              <div className="mode-selectors">
+                <Tabs onTabChange={this.props.onColorModeChange}
+                      onlyHeader
+                      defaultActiveKey={mode}>
+                  {modes.map(mode => <TabPane key={mode}
+                                              id={mode} 
+                                              tab={mode}
+                                              view={() => this.renderModifyIndicator(mode)} />
+                  )}
+                </Tabs>
+              </div>
+            }
             <div className="color-picker-singleton">
               <div className="color-wheel">
                 <FontAwesomeIcon className="reset-icon" icon={['fas', 'undo']} onClick={() => this.onReset(property, reset)} />
                 {colorMode === ColorMode.Rgb && <div className="color-picker-value" style={this.getValueColor()} />}
                 <ColorPicker value={color}
-                             alpha={mode === CCRMode.Color}
                              mode={colorMode}
+                             alpha={mode === CCRMode.Color}
                              type={this.getPropertyType()}
                              max={range.max}
                              onChange={value => this.onColorChange(value ?? reset)} />
@@ -313,7 +314,7 @@ export class CorrectionColorPicker extends React.Component<Props, State> {
               </div>
               <div className="details-section">
                 <div className="color-property-selection">
-                  {mode !== CCRMode.Color && Object.values(ColorProperty).map(property =>
+                  {section === Section.ColorCorrection && Object.values(ColorProperty).map(property =>
                     <div key={property}
                          onClick={() => this.props.onColorPropertyChange(property)}
                          className={`mode-selector ${property === colorProperty ? 'active' : ''}`}>{property}</div>
diff --git a/Engine/Plugins/VirtualProduction/RemoteControlWebInterface/WebApp/_package.json b/Engine/Plugins/VirtualProduction/RemoteControlWebInterface/WebApp/_package.json
--- a/Engine/Plugins/VirtualProduction/RemoteControlWebInterface/WebApp/_package.json
+++ b/Engine/Plugins/VirtualProduction/RemoteControlWebInterface/WebApp/_package.json
@@ -1,6 +1,6 @@
 {
   "name": "conductor",
-  "version": "0.3.64",
+  "version": "0.3.77",
   "scripts": {
     "postinstall": "concurrently \"cd ./Client && npm install && cd ..\" \"cd ./Server && npm install && cd ..\"",
     "build": "node ./Scripts/build.js",
