RemoteControlWebInterface: Merge from Main5

diff --git a/Engine/Plugins/VirtualProduction/RemoteControlWebInterface/WebApp/Start.sh b/Engine/Plugins/VirtualProduction/RemoteControlWebInterface/WebApp/Start.sh
--- a/Engine/Plugins/VirtualProduction/RemoteControlWebInterface/WebApp/Start.sh
+++ b/Engine/Plugins/VirtualProduction/RemoteControlWebInterface/WebApp/Start.sh
@@ -2,6 +2,40 @@
 
 set -e
 
+
+NodeVersion=v16.17.0
+DownloadedNodeFolder=Node-${NodeVersion}
+NodeName=node-${NodeVersion}-linux-x64.tar.gz
+
+if [ ! -d "${DownloadedNodeFolder}" ]; then
+  echo Downloading NodeJS ...
+  
+  # Download nodejs and follow redirects.
+  curl -s -L -o .\${NodeName}.tar.gz "https://nodejs.org/dist/${NodeVersion}/${NodeName}.tar.gz"
+
+  # Only if download succssed
+  if [ test -f "${NodeName}.tar.gz" ]; then
+    # Unarchive the archive
+	tar -xf ${NodeName}.tar.gz
+    
+
+    # Delete the downloaded node.zip
+	rm ${NodeName}.tar.gz
+	
+  else
+    echo Failed to download NodeJS ${NodeVersion}
+  fi
+fi
+
+
+if [ -d "${DownloadedNodeFolder}" ]; then
+  # Add downloaded nodejs version to be first in line
+  echo Using downloaded NodeJS version
+  export PATH="${DownloadedNodeFolder}/bin:$PATH"
+fi
+
+
+
 # First we check if nodejs is installed
 if ! command -v node > /dev/null ; then
   echo "ERROR: Couldn't find node.js installed..., Please install latest nodejs from https://nodejs.org/en/download/"
