- Fixing crash when taking a snapshot from webapp

diff --git a/Engine/Plugins/VirtualProduction/RemoteControlWebInterface/WebApp/Client/src/components/Screens/Snapshot.tsx b/Engine/Plugins/VirtualProduction/RemoteControlWebInterface/WebApp/Client/src/components/Screens/Snapshot.tsx
--- a/Engine/Plugins/VirtualProduction/RemoteControlWebInterface/WebApp/Client/src/components/Screens/Snapshot.tsx
+++ b/Engine/Plugins/VirtualProduction/RemoteControlWebInterface/WebApp/Client/src/components/Screens/Snapshot.tsx
@@ -336,7 +336,7 @@ class SaveModal extends React.Component<SaveModalProps, SaveModalState> {
   async componentDidMount() {
 
     try {
-      const editor = '/Script/LevelSnapshots.Default__LevelSnapshotsEditorDataManagementSettings';
+      const editor = '/Script/LevelSnapshotsEditor.Default__LevelSnapshotsEditorSettings';
 
       const { RootLevelSnapshotSaveDir } = await _api.proxy.property.get(editor, 'RootLevelSnapshotSaveDir');
       let path: string = RootLevelSnapshotSaveDir?.Path ?? '/Game/LevelSnapshots';
diff --git a/Engine/Plugins/VirtualProduction/RemoteControlWebInterface/WebApp/_package.json b/Engine/Plugins/VirtualProduction/RemoteControlWebInterface/WebApp/_package.json
--- a/Engine/Plugins/VirtualProduction/RemoteControlWebInterface/WebApp/_package.json
+++ b/Engine/Plugins/VirtualProduction/RemoteControlWebInterface/WebApp/_package.json
@@ -1,6 +1,6 @@
 {
   "name": "conductor",
-  "version": "0.3.77",
+  "version": "0.3.78",
   "scripts": {
     "postinstall": "concurrently \"cd ./Client && npm install && cd ..\" \"cd ./Server && npm install && cd ..\"",
     "build": "node ./Scripts/build.js",
