RemoteControlWebInterface: Added support for Controllers

diff --git a/Engine/Plugins/VirtualProduction/RemoteControlWebInterface/WebApp/Client/src/shared/index.ts b/Engine/Plugins/VirtualProduction/RemoteControlWebInterface/WebApp/Client/src/shared/index.ts
--- a/Engine/Plugins/VirtualProduction/RemoteControlWebInterface/WebApp/Client/src/shared/index.ts
+++ b/Engine/Plugins/VirtualProduction/RemoteControlWebInterface/WebApp/Client/src/shared/index.ts
@@ -102,6 +102,21 @@ export interface IExposedProperty {
   Type: PropertyType;
 }
 
+export interface IController {
+  ID: string;
+  Name: string;
+  Type: PropertyType;
+  Path: string;
+
+  DisplayName: string;
+  Metadata: Record<string, string>;
+  Widget: WidgetType;
+
+
+  UnderlyingProperty: IProperty;
+  OwnerObjects: IObject[];
+}
+
 export interface IActor {
   Name: string;
   Path: string;
@@ -127,7 +142,8 @@ export interface IPreset {
   
   ExposedProperties?: IExposedProperty[];
   ExposedFunctions?: IExposedFunction[];
-  Exposed?: Record<string, IExposedProperty | IExposedFunction>;
+  Controllers?: IController[];
+  Exposed?: Record<string, IExposedProperty | IExposedFunction | IController>;
   IsFavorite?: boolean;
 }
 
diff --git a/Engine/Plugins/VirtualProduction/RemoteControlWebInterface/WebApp/Client/src/utilities/WidgetUtilities.ts b/Engine/Plugins/VirtualProduction/RemoteControlWebInterface/WebApp/Client/src/utilities/WidgetUtilities.ts
--- a/Engine/Plugins/VirtualProduction/RemoteControlWebInterface/WebApp/Client/src/utilities/WidgetUtilities.ts
+++ b/Engine/Plugins/VirtualProduction/RemoteControlWebInterface/WebApp/Client/src/utilities/WidgetUtilities.ts
@@ -153,6 +153,9 @@ export class WidgetUtilities {
 
       case PropertyType.Rotator:
         return [WidgetTypes.Vector, WidgetTypes.Sliders, WidgetTypes.Dials];
+
+      case PropertyType.String:
+        return [WidgetTypes.Text];
     }
 
     if (propertyType && (propertyType.startsWith('TEnum') || propertyType.startsWith('E')))
diff --git a/Engine/Plugins/VirtualProduction/RemoteControlWebInterface/WebApp/Server/src/UnrealEngine.ts b/Engine/Plugins/VirtualProduction/RemoteControlWebInterface/WebApp/Server/src/UnrealEngine.ts
--- a/Engine/Plugins/VirtualProduction/RemoteControlWebInterface/WebApp/Server/src/UnrealEngine.ts
+++ b/Engine/Plugins/VirtualProduction/RemoteControlWebInterface/WebApp/Server/src/UnrealEngine.ts
@@ -1,5 +1,5 @@
 import { IPayload, IPayloads, IPreset, IPresets, IPanel, IView, ICustomStackWidget, ICustomStackTabs, PropertyValue, 
-          IAsset, WidgetTypes, PropertyType, IColorPickerList, ICustomStackItem } from '../../Client/src/shared';
+          IAsset, WidgetTypes, PropertyType, IColorPickerList, ICustomStackItem, IGroup, IExposedProperty } from '../../Client/src/shared';
 import _ from 'lodash';
 import WebSocket from 'ws';
 import { Notify, Program, LogServer } from './';
@@ -9,14 +9,19 @@ import crypto from 'crypto';
 
 namespace UnrealApi {
   export enum PresetEvent {
-    FieldsRenamed     = 'PresetFieldsRenamed',
-    FieldsChanged     = 'PresetFieldsChanged',
-    FieldsAdded       = 'PresetFieldsAdded',
-    FieldsRemoved     = 'PresetFieldsRemoved',
-    MetadataModified  = 'PresetMetadataModified',
-    ActorModified     = 'PresetActorModified',
-    EntitiesModified  = 'PresetEntitiesModified',
-    LayoutModified    = 'PresetLayoutModified',
+    FieldsRenamed         = 'PresetFieldsRenamed',
+    FieldsChanged         = 'PresetFieldsChanged',
+    FieldsAdded           = 'PresetFieldsAdded',
+    FieldsRemoved         = 'PresetFieldsRemoved',
+    MetadataModified      = 'PresetMetadataModified',
+    ActorModified         = 'PresetActorModified',
+    EntitiesModified      = 'PresetEntitiesModified',
+    LayoutModified        = 'PresetLayoutModified',
+
+    ControllersAdded      = 'PresetControllersAdded',
+    ControllersRemoved    = 'PresetControllersRemoved',
+    ControllersModified   = 'PresetControllersModified',
+    ControllersRenamed    = 'PresetControllersRenamed',
   }
 
   export type Presets = { 
@@ -100,7 +105,7 @@ export namespace UnrealEngine {
     }
 
     connect();
-    startQuitTimeout(60 * 1000);
+    startQuitTimeout(10 * 60 * 1000);
   }
 
   export function isConnected(): boolean {
@@ -288,11 +293,7 @@ export namespace UnrealEngine {
           break;
         }
 
-        case UnrealApi.PresetEvent.FieldsRemoved: {
-          await refreshPreset(message.PresetId, message.PresetName);
-          break;
-        }
-
+        case UnrealApi.PresetEvent.FieldsRemoved:
         case UnrealApi.PresetEvent.EntitiesModified: {
           await refreshPreset(message.PresetId, message.PresetName);
           break;
@@ -316,6 +317,34 @@ export namespace UnrealEngine {
 
           break;
         }
+
+        case UnrealApi.PresetEvent.ControllersAdded:
+        case UnrealApi.PresetEvent.ControllersRemoved:
+        case UnrealApi.PresetEvent.ControllersRenamed: {
+          await refreshPreset(message.PresetId, message.PresetName);
+          break;
+        }
+
+        case UnrealApi.PresetEvent.ControllersModified: {
+          const preset = presets[message.PresetId];
+          if (!preset)
+            break;
+
+          const changes: any = {};
+          for (let i = 0; i < message.ModifiedControllers.Controllers.length; i++) {
+            const id = message.ModifiedControllers.Controllers[i].ID;
+            const controller = preset.Exposed[id];
+            if (!controller)
+              continue;
+
+            const value = message.ModifiedControllers.ChangedValues[i];
+            setPayloadValueInternal(payloads, [message.PresetId, id], value);
+            changes[id] = value;
+          }
+
+          Notify.emitValuesChanges(message.PresetId, changes);          
+          break;
+        }
       }
 
     } catch (error) {
@@ -493,8 +522,26 @@ export namespace UnrealEngine {
       preset.ExposedFunctions.push(...Group.ExposedFunctions);
     }
 
-    preset.IsFavorite = favorites?.[preset.ID];
+    if (preset.Controllers) {
+      const Controllers: IGroup = { Name: 'Controllers', ExposedProperties: [], ExposedFunctions: [] };
+      preset.Groups.push(Controllers);
+      for (const Controller of preset.Controllers) {
+        if (!Controller.Metadata)
+          Controller.Metadata = {};
+
+        Controller.UnderlyingProperty = {
+          Name: Controller.Name,
+          Metadata: Controller.Metadata,
+          Description: '',
+          Type: Controller.Type,
+        };
+
+        preset.Exposed[Controller.ID] = Controller;
+        Controllers.ExposedProperties.push(Controller);
+      }
+    }
 
+    preset.IsFavorite = favorites?.[preset.ID];
     return preset;
   }
 
@@ -605,24 +652,41 @@ export namespace UnrealEngine {
     return false;
   }
 
+  function getPropertyValueRequest(Preset: IPreset, Property: IExposedProperty, IsController: boolean) {
+    return ({
+      RequestId: httpRequest++,
+      Verb: 'GET',
+      URL: `/remote/preset/${Preset.ID}/property/${Property.ID}`,
+      PropertyId: Property.ID,
+      IsController,
+    });
+  }
+
   async function pullPresetValues(Preset: IPreset): Promise<IPayload> {
     const Requests = [];
-    for (const property of Preset.ExposedProperties) {
-      Requests.push({
-        RequestId: httpRequest++,
-        Verb: 'GET',
-        URL: `/remote/preset/${Preset.ID}/property/${property.ID}`,
-        PropertyId: property.ID,
-      });
+    for (const Property of Preset.ExposedProperties)
+      Requests.push(getPropertyValueRequest(Preset, Property, false));
+
+    if (Preset.Controllers) {
+      for (const Controller of Preset.Controllers)
+        Requests.push(getPropertyValueRequest(Preset, Controller, true));
     }
 
     const updatedPayloads: IPayloads = {};
-    const values = await put<UnrealApi.BatchResponses>('/remote/batch', { Requests });
+    const answers = await put<UnrealApi.BatchResponses>('/remote/batch', { Requests });
     for (let i = 0; i < Requests.length; i++) {
-      const { PropertyId } = Requests[i];
+      const { PropertyId, IsController } = Requests[i];
       try {
-        const value = values.Responses[i].ResponseBody as UnrealApi.PropertyValues;
-        setPayloadValueInternal(updatedPayloads, [Preset.ID, PropertyId], value?.PropertyValues?.[0]?.PropertyValue);
+        let value;
+        if (IsController) {
+          const body = answers.Responses[i].ResponseBody as UnrealApi.GetPropertyValue;
+          value = body?.PropertyValue;
+        } else {
+          const body = answers.Responses[i].ResponseBody as UnrealApi.PropertyValues;
+          value = body?.PropertyValues?.[0]?.PropertyValue;
+        }
+
+        setPayloadValueInternal(updatedPayloads, [Preset.ID, PropertyId], value);
       } catch (error) {
         console.log(`Failed to get value of Preset: ${Preset.Name}, Property: ${PropertyId}`);
       }
diff --git a/Engine/Plugins/VirtualProduction/RemoteControlWebInterface/WebApp/_package.json b/Engine/Plugins/VirtualProduction/RemoteControlWebInterface/WebApp/_package.json
--- a/Engine/Plugins/VirtualProduction/RemoteControlWebInterface/WebApp/_package.json
+++ b/Engine/Plugins/VirtualProduction/RemoteControlWebInterface/WebApp/_package.json
@@ -1,6 +1,6 @@
 {
   "name": "conductor",
-  "version": "0.3.78",
+  "version": "0.3.79",
   "scripts": {
     "postinstall": "concurrently \"cd ./Client && npm install && cd ..\" \"cd ./Server && npm install && cd ..\"",
     "build": "node ./Scripts/build.js",
