RemoteControlWebInterface: Link CCR to the updated C++ CCR

diff --git a/Engine/Plugins/VirtualProduction/RemoteControlWebInterface/WebApp/Client/src/components/Screens/ColorCorrection/ColorCorrection.tsx b/Engine/Plugins/VirtualProduction/RemoteControlWebInterface/WebApp/Client/src/components/Screens/ColorCorrection/ColorCorrection.tsx
--- a/Engine/Plugins/VirtualProduction/RemoteControlWebInterface/WebApp/Client/src/components/Screens/ColorCorrection/ColorCorrection.tsx
+++ b/Engine/Plugins/VirtualProduction/RemoteControlWebInterface/WebApp/Client/src/components/Screens/ColorCorrection/ColorCorrection.tsx
@@ -162,7 +162,7 @@ export class ColorCorrection extends React.Component<Props, State> {
 
     switch (section) {
       case Section.ColorCorrection:
-        Class = '/ColorCorrectRegions/Blueprints/ColorCorrectRegion.ColorCorrectRegion_C';
+        Class = '/Script/ColorCorrectRegions.ColorCorrectionRegion';
         break;
 
       case Section.LightCards:
diff --git a/Engine/Plugins/VirtualProduction/RemoteControlWebInterface/WebApp/_package.json b/Engine/Plugins/VirtualProduction/RemoteControlWebInterface/WebApp/_package.json
--- a/Engine/Plugins/VirtualProduction/RemoteControlWebInterface/WebApp/_package.json
+++ b/Engine/Plugins/VirtualProduction/RemoteControlWebInterface/WebApp/_package.json
@@ -1,6 +1,6 @@
 {
   "name": "conductor",
-  "version": "0.3.79",
+  "version": "0.3.80",
   "scripts": {
     "postinstall": "concurrently \"cd ./Client && npm install && cd ..\" \"cd ./Server && npm install && cd ..\"",
     "build": "node ./Scripts/build.js",
