RemoteControlWebInterface: Download nodejs binary in linux / Mac

diff --git a/Engine/Plugins/VirtualProduction/RemoteControlWebInterface/WebApp/Start.sh b/Engine/Plugins/VirtualProduction/RemoteControlWebInterface/WebApp/Start.sh
--- a/Engine/Plugins/VirtualProduction/RemoteControlWebInterface/WebApp/Start.sh
+++ b/Engine/Plugins/VirtualProduction/RemoteControlWebInterface/WebApp/Start.sh
@@ -2,25 +2,29 @@
 
 set -e
 
+NodeOS=linux
+if [[ $OSTYPE == 'darwin'* ]]; then
+  NodeOS=darwin
+fi
 
 NodeVersion=v16.17.0
-DownloadedNodeFolder=Node-${NodeVersion}
-NodeName=node-${NodeVersion}-linux-x64.tar.gz
+DownloadedNodeFolder=node-${NodeVersion}-${NodeOS}-x64
+NodeName=${DownloadedNodeFolder}.tar.gz
 
 if [ ! -d "${DownloadedNodeFolder}" ]; then
   echo Downloading NodeJS ...
   
   # Download nodejs and follow redirects.
-  curl -s -L -o .\${NodeName}.tar.gz "https://nodejs.org/dist/${NodeVersion}/${NodeName}.tar.gz"
+  curl -s -L -o ./${NodeName} "https://nodejs.org/dist/${NodeVersion}/${NodeName}"
 
   # Only if download succssed
-  if [ test -f "${NodeName}.tar.gz" ]; then
+  if [ -f "${NodeName}" ]; then
     # Unarchive the archive
-	tar -xf ${NodeName}.tar.gz
+	tar -xf ${NodeName}
     
 
     # Delete the downloaded node.zip
-	rm ${NodeName}.tar.gz
+	rm ${NodeName}
 	
   else
     echo Failed to download NodeJS ${NodeVersion}
