RemoteControlWebInterface: Adjusting WebApp to match 5.1 CCR/Lightcard propreties

diff --git a/Engine/Plugins/VirtualProduction/RemoteControlWebInterface/WebApp/Client/src/components/Screens/ColorCorrection/ColorCorrection.tsx b/Engine/Plugins/VirtualProduction/RemoteControlWebInterface/WebApp/Client/src/components/Screens/ColorCorrection/ColorCorrection.tsx
--- a/Engine/Plugins/VirtualProduction/RemoteControlWebInterface/WebApp/Client/src/components/Screens/ColorCorrection/ColorCorrection.tsx
+++ b/Engine/Plugins/VirtualProduction/RemoteControlWebInterface/WebApp/Client/src/components/Screens/ColorCorrection/ColorCorrection.tsx
@@ -630,6 +630,26 @@ export class ColorCorrection extends React.Component<Props, State> {
     );
   }
 
+  renderAdvancedContent = (pagination: number): Element[] => {
+    const contents = [];
+
+    if (pagination === 0)
+      contents.push(
+        this.renderSlider('ColorGradingSettings.ShadowsMax', 0, 1, 0),
+        this.renderSlider('ColorGradingSettings.HighlightsMax', 1, 10, 0),
+        this.renderSlider('ColorGradingSettings.HighlightsMin', -1, 1, 0),
+      );
+
+    if (pagination === 1)
+      contents.push(
+        this.renderSlider('Inner', 0, 1, 0.5),
+        this.renderSlider('Outer', 0, 1, 1, 3),
+        this.renderSlider('Priority', -1, 20, 0, 0),
+      );
+
+    return contents;
+  }
+
   renderAppearanceContent = (pagination: number): Element[] => {
     const contents = [];
 
@@ -645,31 +665,32 @@ export class ColorCorrection extends React.Component<Props, State> {
       contents.push(
         this.renderSlider('Opacity', 0, 1, 1, 3),
         this.renderSlider('Feathering', 0, 1, 0, 3),
-        <div />,
-        <div />,
+        <div key="none1" />,
+        <div key="none2" />,
       );
 
     return contents;
   }
 
   renderOrientationContent = (pagination: number) => {
-    const { value } = this.state;
     const contents = [];
 
     if (pagination === 0)
       contents.push(
         this.renderSlider('DistanceFromCenter', -2_000, 2_000, 300),
-        this.renderSlider('Longitude', 0, 360, 0, 2, null, _.get(value, 'UseClockSystem')),
+        this.renderSlider('Longitude', 0, 360, 0, 2),
         this.renderSlider('Latitude', -90, 90, 0),
-        this.renderSlider('Spin', 0, 360, 0),
+        this.renderSlider('Spin', -360, 360, 0),
+        this.renderSlider('Pitch', -360, 360, 0),
+        this.renderSlider('Yaw', -360, 360, 0),
       );
 
     if (pagination === 1)
       contents.push(
         this.renderSlider('Scale.X', 0, 20, 2, 3, 'Scale X'),
         this.renderSlider('Scale.Y', 0, 20, 2, 3, 'Scale Y'),
-        <div />,
-        <div />,
+        this.renderSlider('RadialOffset', -360, 360, 0),
+        <div key="empty" />,
       );
 
     return contents;
@@ -722,26 +743,26 @@ export class ColorCorrection extends React.Component<Props, State> {
               {this.renderPositionTab()}
             </TabPane>
             <TabPane id="default" tab="Default">
-              <div className='props-row'>
+              <div className='props-row fit-dropdown'>
                 {this.renderSelect('Type', ['Sphere', 'Box', 'Cylinder', 'Cone'])}
+                {this.renderSelect('TemperatureType', ['LegacyTemperature', 'WhiteBalance', 'ColorTemperature'])}
                 {this.renderToggle('Enabled')}
               </div>
               <div className='control-rows'>
                 {this.renderSlider('Temperature', 0, 10_000, 6_500)}
                 {this.renderSlider('Intensity', 0, 1, 1)}
                 {this.renderSlider('Falloff', 0, 20, 0, 3)}
-                {this.renderSlider('Inner', 0, 1, 0.5)}
+                {this.renderSlider('Tint', -1, 1, 0, 3)}
               </div>
             </TabPane>
             <TabPane id="advanced" tab="Advanced">
               <div className='props-row'>
                 {this.renderToggle('Invert')}
               </div>
-              <div className='control-rows'>
-                {this.renderSlider('ColorGradingSettings.ShadowsMax', 0, 1, 0)}
-                {this.renderSlider('ColorGradingSettings.HighlightsMin', 0, 1, 0)}
-                {this.renderSlider('Outer', 0, 1, 1, 3)}
-                {this.renderSlider('Priority', -1, 20, 0, 0)}
+              <div className='multilines-content'>
+                <PaginationContent pages={2} key="advanced">
+                  {({ pagination }) => this.renderAdvancedContent(pagination)}
+                </PaginationContent>
               </div>
             </TabPane>
           </Tabs>
@@ -773,13 +794,13 @@ export class ColorCorrection extends React.Component<Props, State> {
               </div>
             </TabPane>
             <TabPane id="orientation" tab="Orientation">
-              <div className="multilines-content">
+              <div className="multilines-content full-height">
                 <PaginationContent pages={2} key="orientation">
                   {({ pagination }) => this.renderOrientationContent(pagination)}
                 </PaginationContent>
               </div>
             </TabPane>
-            <TabPane id="guides" tab="Guides">
+            {/* <TabPane id="guides" tab="Guides">
               <div className='props-row'>
                 {this.renderToggle('bShowGuides', 'Show Guides')}
               </div>
@@ -787,7 +808,7 @@ export class ColorCorrection extends React.Component<Props, State> {
                 {this.renderSlider('GuideScale', 0, 1, 0)}
                 {this.renderGuidesColorPicker()}
               </div>
-            </TabPane>
+            </TabPane> */}
             <TabPane id="position" tab="Position">
               {this.renderPositionTab()}
             </TabPane>
diff --git a/Engine/Plugins/VirtualProduction/RemoteControlWebInterface/WebApp/Client/src/components/Screens/ColorCorrection/CorrectionColorPicker.tsx b/Engine/Plugins/VirtualProduction/RemoteControlWebInterface/WebApp/Client/src/components/Screens/ColorCorrection/CorrectionColorPicker.tsx
--- a/Engine/Plugins/VirtualProduction/RemoteControlWebInterface/WebApp/Client/src/components/Screens/ColorCorrection/CorrectionColorPicker.tsx
+++ b/Engine/Plugins/VirtualProduction/RemoteControlWebInterface/WebApp/Client/src/components/Screens/ColorCorrection/CorrectionColorPicker.tsx
@@ -334,20 +334,6 @@ export class CorrectionColorPicker extends React.Component<Props, State> {
                                     value={hue}
                                     onChange={hue => this.onHsvWheelMove(1, hue)} />
                       </div>
-                      {(mode === CCRMode.Shadows) &&
-                        <div className="slider-wheel-container shadow-min-slider">
-                          <span className="title">Shadow Min</span>
-                          <SliderWheel size={200}
-                                      className="infinite-slider" 
-                                      onWheelMove={sign => {}} />
-                          <ValueInput draggable={false}
-                                      precision={1}
-                                      min={0}
-                                      max={360}
-                                      value={0}
-                                      onChange={v => {}} />
-                        </div>
-                      }
                     </div>
                   </div>
                 </div>
diff --git a/Engine/Plugins/VirtualProduction/RemoteControlWebInterface/WebApp/_package.json b/Engine/Plugins/VirtualProduction/RemoteControlWebInterface/WebApp/_package.json
--- a/Engine/Plugins/VirtualProduction/RemoteControlWebInterface/WebApp/_package.json
+++ b/Engine/Plugins/VirtualProduction/RemoteControlWebInterface/WebApp/_package.json
@@ -1,6 +1,6 @@
 {
   "name": "conductor",
-  "version": "0.3.80",
+  "version": "0.3.81",
   "scripts": {
     "postinstall": "concurrently \"cd ./Client && npm install && cd ..\" \"cd ./Server && npm install && cd ..\"",
     "build": "node ./Scripts/build.js",
