RemoteControlWebInterface: Allow to delete actors from panel

diff --git a/Engine/Plugins/VirtualProduction/RemoteControlWebInterface/WebApp/Client/src/components/Screens/ColorCorrection/ColorCorrection.tsx b/Engine/Plugins/VirtualProduction/RemoteControlWebInterface/WebApp/Client/src/components/Screens/ColorCorrection/ColorCorrection.tsx
--- a/Engine/Plugins/VirtualProduction/RemoteControlWebInterface/WebApp/Client/src/components/Screens/ColorCorrection/ColorCorrection.tsx
+++ b/Engine/Plugins/VirtualProduction/RemoteControlWebInterface/WebApp/Client/src/components/Screens/ColorCorrection/ColorCorrection.tsx
@@ -1,7 +1,7 @@
 import React from 'react';
 import { _api } from 'src/reducers';
 import _ from 'lodash';
-import { Button, ColorPicker, Search, Slider, TabPane, Tabs, VectorControl, TimeControl, SliderWheel, ValueInput } from '../../controls';
+import { Button, ColorPicker, Search, Slider, TabPane, Tabs, VectorControl, SliderWheel, ValueInput } from '../../controls';
 import { AlertModal } from '../../../components';
 import { AssetWidget } from '../../Widgets';
 import { CorrectionColorPicker, PaginationContent } from './';
@@ -182,6 +182,18 @@ export class ColorCorrection extends React.Component<Props, State> {
     await this.onRefreshActors();
   }
 
+  onDeleteActor = async() => {
+    const { selected, section } = this.state;
+    if (!selected)
+      return;
+
+    if (!await AlertModal.show(`Delete ${section === Section.ColorCorrection ? 'Color Correct Region' : 'Light Card'}?`))
+      return;
+
+    await _api.proxy.function(selected, 'DestroyActor');
+    this.onRefreshActors();
+  }
+
   onRefreshActors = async () => {
     const { section } = this.state;
     let Class;
@@ -856,7 +868,10 @@ export class ColorCorrection extends React.Component<Props, State> {
         </div>
         <div className="group">
           <div className="cc-group shotbox cc-panel">
-            <Tabs rightIcon={<Search placeholder="Search" onSearch={search => this.setState({ search })} />}
+            <Tabs rightIcon={<>
+              <div className="trash-icon" onClick={this.onDeleteActor}><FontAwesomeIcon icon={['fas', 'trash']} /></div>
+              <Search placeholder="Search" onSearch={search => this.setState({ search })} />
+            </>}
                   onTabChange={(actorsGroup: ActorsGroup) => this.setState({ actorsGroup })}
                   defaultActiveKey={actorsGroup}>
               <TabPane id={ActorsGroup.All} tab="All" icon={['fas', 'square']}>
diff --git a/Engine/Plugins/VirtualProduction/RemoteControlWebInterface/WebApp/_package.json b/Engine/Plugins/VirtualProduction/RemoteControlWebInterface/WebApp/_package.json
--- a/Engine/Plugins/VirtualProduction/RemoteControlWebInterface/WebApp/_package.json
+++ b/Engine/Plugins/VirtualProduction/RemoteControlWebInterface/WebApp/_package.json
@@ -1,6 +1,6 @@
 {
   "name": "conductor",
-  "version": "0.3.81",
+  "version": "0.3.82",
   "scripts": {
     "postinstall": "concurrently \"cd ./Client && npm install && cd ..\" \"cd ./Server && npm install && cd ..\"",
     "build": "node ./Scripts/build.js",
