Fix web interface sometimes not reacting to preset view change

diff --git a/Engine/Plugins/VirtualProduction/RemoteControlWebInterface/WebApp/Server/src/UnrealEngine.ts b/Engine/Plugins/VirtualProduction/RemoteControlWebInterface/WebApp/Server/src/UnrealEngine.ts
--- a/Engine/Plugins/VirtualProduction/RemoteControlWebInterface/WebApp/Server/src/UnrealEngine.ts
+++ b/Engine/Plugins/VirtualProduction/RemoteControlWebInterface/WebApp/Server/src/UnrealEngine.ts
@@ -1,5 +1,5 @@
 import { IPayload, IPayloads, IPreset, IPresets, IPanel, IView, ICustomStackWidget, ICustomStackTabs, PropertyValue, 
-          IAsset, WidgetTypes, PropertyType, IColorPickerList, ICustomStackItem, IGroup, IExposedProperty } from '../../Client/src/shared';
+          IAsset, WidgetTypes, PropertyType, IColorPickerList, ICustomStackItem, IGroup, IExposedProperty, TabLayout } from '../../Client/src/shared';
 import _ from 'lodash';
 import WebSocket from 'ws';
 import { Notify, Program, LogServer } from './';
@@ -300,9 +300,6 @@ export namespace UnrealEngine {
         }
 
         case UnrealApi.PresetEvent.MetadataModified: {
-          if (!message.Metadata.view)
-            break;
-
           await refreshPreset(message.PresetId, message.PresetName);
           refreshView(message.PresetId, message.Metadata.view);
           break;
@@ -547,8 +544,14 @@ export namespace UnrealEngine {
 
   function refreshView(id: string, viewJson: string) {
     try {
-      if (!viewJson)
+      if (!viewJson) {
+        views[id] = {
+          tabs: [{ name: 'Tab 1', layout: TabLayout.Empty, icon: '' }]
+        };
+
+        Notify.onViewChange(id, views[id], true);
         return;
+      }
 
       const view = JSON.parse(viewJson) as IView;
       if (equal(view, views[id]))
diff --git a/Engine/Plugins/VirtualProduction/RemoteControlWebInterface/WebApp/_package.json b/Engine/Plugins/VirtualProduction/RemoteControlWebInterface/WebApp/_package.json
--- a/Engine/Plugins/VirtualProduction/RemoteControlWebInterface/WebApp/_package.json
+++ b/Engine/Plugins/VirtualProduction/RemoteControlWebInterface/WebApp/_package.json
@@ -1,6 +1,6 @@
 {
   "name": "conductor",
-  "version": "0.3.82",
+  "version": "0.3.83",
   "scripts": {
     "postinstall": "concurrently \"cd ./Client && npm install && cd ..\" \"cd ./Server && npm install && cd ..\"",
     "build": "node ./Scripts/build.js",
