Fix web app not building

diff --git a/Engine/Plugins/VirtualProduction/RemoteControlWebInterface/WebApp/Client/src/components/Stack.tsx b/Engine/Plugins/VirtualProduction/RemoteControlWebInterface/WebApp/Client/src/components/Stack.tsx
--- a/Engine/Plugins/VirtualProduction/RemoteControlWebInterface/WebApp/Client/src/components/Stack.tsx
+++ b/Engine/Plugins/VirtualProduction/RemoteControlWebInterface/WebApp/Client/src/components/Stack.tsx
@@ -764,7 +764,7 @@ export class Stack extends React.Component<Props & PropsFromState, State> {
         <Widgets.AssetWidget label={this.getLabel(widget)}
                              value={this.getValue(widget)}
                              type={exposed?.Type ?? 'Object'}
-                             typePath={exposed.TypePath}
+                             typePath={exposed?.TypePath}
                              onChange={value => this.onPropertyValueChange(widget, value)}
                              browse />
       </DraggableWidget>
diff --git a/Engine/Plugins/VirtualProduction/RemoteControlWebInterface/WebApp/Client/src/components/Widgets/AssetWidget.tsx b/Engine/Plugins/VirtualProduction/RemoteControlWebInterface/WebApp/Client/src/components/Widgets/AssetWidget.tsx
--- a/Engine/Plugins/VirtualProduction/RemoteControlWebInterface/WebApp/Client/src/components/Widgets/AssetWidget.tsx
+++ b/Engine/Plugins/VirtualProduction/RemoteControlWebInterface/WebApp/Client/src/components/Widgets/AssetWidget.tsx
@@ -9,6 +9,7 @@ type Props = {
   label?: string | JSX.Element;
   value?: string;
   type?: string;
+  typePath?: string;
   browse?: boolean;
   reset?: boolean;
   onChange?: (value?: string) => void;
@@ -21,13 +22,15 @@ export class AssetWidget extends React.Component<Props> {
   };
 
   onBrowse = async () => {
-    let { type } = this.props;
+    let { type, typePath } = this.props;
     if (type.startsWith('U') || type.startsWith('F'))
       type = type.substring(1);
 
     type = type.replace('*', '');
+    if (typePath?.length == 0)
+      typePath = type;
 
-    const asset = await SearchModal.open({ placeholder: 'Search asset', types: [type], prefix: '/Game', count: 200 });
+    const asset = await SearchModal.open({ placeholder: 'Search asset', types: [typePath], prefix: '/Game', count: 200 });
     if (asset)
       this.props.onChange?.(asset.Path);
   }
diff --git a/Engine/Plugins/VirtualProduction/RemoteControlWebInterface/WebApp/_package.json b/Engine/Plugins/VirtualProduction/RemoteControlWebInterface/WebApp/_package.json
--- a/Engine/Plugins/VirtualProduction/RemoteControlWebInterface/WebApp/_package.json
+++ b/Engine/Plugins/VirtualProduction/RemoteControlWebInterface/WebApp/_package.json
@@ -1,6 +1,6 @@
 {
   "name": "conductor",
-  "version": "0.3.84 ",
+  "version": "0.3.85 ",
   "scripts": {
     "postinstall": "concurrently \"cd ./Client && npm install && cd ..\" \"cd ./Server && npm install && cd ..\"",
     "build": "node ./Scripts/build.js",
