RemoteControlWebInterface: - Rotate Color Picker Wheel to match UE
- Dropdown background color is white

diff --git a/Engine/Plugins/VirtualProduction/RemoteControlWebInterface/WebApp/Client/src/components/controls/ColorPicker.tsx b/Engine/Plugins/VirtualProduction/RemoteControlWebInterface/WebApp/Client/src/components/controls/ColorPicker.tsx
--- a/Engine/Plugins/VirtualProduction/RemoteControlWebInterface/WebApp/Client/src/components/controls/ColorPicker.tsx
+++ b/Engine/Plugins/VirtualProduction/RemoteControlWebInterface/WebApp/Client/src/components/controls/ColorPicker.tsx
@@ -182,7 +182,7 @@ export class ColorPicker extends React.Component<Props, State> {
     const maxRadius = this.circleRef.current.clientHeight / 2;
 
     const radians = Math.atan2(x, y);
-    const h = this.normalize(-1 * radians * (180 / Math.PI) - 90) / 360;
+    const h = this.normalize(-1 * radians * (180 / Math.PI) - 270) / 360;
 
     let s = Math.min(radius / maxRadius, 1);
     if (s < 0.02)
@@ -391,8 +391,8 @@ export class ColorPicker extends React.Component<Props, State> {
     if (mode === ColorMode.Rgb) {
       const maxCircleRadius = this.circleRef.current?.clientHeight / 2;
 
-      const circleX = Math.sin(radian) * Math.max(0, Math.min(hsv.s, 1)) * (maxCircleRadius - 10);
-      const circleY = -Math.cos(radian) * Math.max(0, Math.min(hsv.s, 1)) * (maxCircleRadius - 10);
+      const circleX = Math.sin(radian) * Math.max(0, Math.min(hsv.s, 1)) * (maxCircleRadius - 10) * -1;
+      const circleY = -Math.cos(radian) * Math.max(0, Math.min(hsv.s, 1)) * (maxCircleRadius - 10) * -1;
 
       return { circleX, circleY };
     }
diff --git a/Engine/Plugins/VirtualProduction/RemoteControlWebInterface/WebApp/_package.json b/Engine/Plugins/VirtualProduction/RemoteControlWebInterface/WebApp/_package.json
--- a/Engine/Plugins/VirtualProduction/RemoteControlWebInterface/WebApp/_package.json
+++ b/Engine/Plugins/VirtualProduction/RemoteControlWebInterface/WebApp/_package.json
@@ -1,6 +1,6 @@
 {
   "name": "conductor",
-  "version": "0.3.85 ",
+  "version": "0.3.86",
   "scripts": {
     "postinstall": "concurrently \"cd ./Client && npm install && cd ..\" \"cd ./Server && npm install && cd ..\"",
     "build": "node ./Scripts/build.js",
