Fix ICVFX plugin test script path

diff --git a/Engine/Plugins/VirtualProduction/ICVFXTesting/Scripts/ICVFXTesting.xml b/Engine/Plugins/VirtualProduction/ICVFXTesting/Scripts/ICVFXTesting.xml
--- a/Engine/Plugins/VirtualProduction/ICVFXTesting/Scripts/ICVFXTesting.xml
+++ b/Engine/Plugins/VirtualProduction/ICVFXTesting/Scripts/ICVFXTesting.xml
@@ -35,7 +35,7 @@
 	<Property Name="ClientCompileAgentOverride" Value="EditorPerfWin64" />
 	<!-- Common Setting Overrides -->
 	
-	<Property Name="GauntletSettingsFile" Value="$(RootDir)/Engine/Plugins/VirtualProduction/ICVFX/Scripts/GauntletSettings.xml" />
+	<Property Name="GauntletSettingsFile" Value="$(RootDir)/Engine/Plugins/VirtualProduction/ICVFXTesting/Scripts/GauntletSettings.xml" />
 
 	<!-- This will declare an aggregate called BuildAndTest $(ProjectName) -->
 	<Include Script="$(RootDir)/Engine/Build/Graph/Tasks/BuildAndTestProject.xml" />
