SaloonPerf
Fix artifact path
Fix UHT warning in lightingtoolbplibrary

diff --git a/Engine/Plugins/VirtualProduction/ICVFXTesting/Build/Scripts/Automation/ICVFX.TestNode.cs b/Engine/Plugins/VirtualProduction/ICVFXTesting/Build/Scripts/Automation/ICVFX.TestNode.cs
--- a/Engine/Plugins/VirtualProduction/ICVFXTesting/Build/Scripts/Automation/ICVFX.TestNode.cs
+++ b/Engine/Plugins/VirtualProduction/ICVFXTesting/Build/Scripts/Automation/ICVFX.TestNode.cs
@@ -10,6 +10,12 @@ namespace ICVFXTest
 {
 	public class ICVFXTestConfig : UnrealTestConfiguration
 	{
+		/// <summary>
+		/// Name of the project.
+		/// </summary>
+		[AutoParamWithNames("ICVFXTest.ProjectName")]
+		public string ProjectName;
+
 		/// <summary>
 		/// Name of the test.
 		/// </summary>
diff --git a/Engine/Plugins/VirtualProduction/ICVFXTesting/Build/Scripts/Automation/Tests/ICVFXTest.AutoTest.cs b/Engine/Plugins/VirtualProduction/ICVFXTesting/Build/Scripts/Automation/Tests/ICVFXTest.AutoTest.cs
--- a/Engine/Plugins/VirtualProduction/ICVFXTesting/Build/Scripts/Automation/Tests/ICVFXTest.AutoTest.cs
+++ b/Engine/Plugins/VirtualProduction/ICVFXTesting/Build/Scripts/Automation/Tests/ICVFXTest.AutoTest.cs
@@ -49,7 +49,7 @@ namespace ICVFXTest
 			Config.MaxDuration = Context.TestParams.ParseValue("MaxDuration", 60 * 60);  // 1 hour max
 
 			UnrealTestRole ClientRole = Config.RequireRole(UnrealTargetRole.Client);
-            ClientRole.Controllers.Add("AutoTest");
+            ClientRole.Controllers.Add("ICVFXTestControllerAutoTest");
 
 			ClientRole.CommandLineParams.AddOrAppendParamValue("execcmds", "ICVFXTest.MaxRunCount " + Config.MaxRunCount);
 
diff --git a/Engine/Plugins/VirtualProduction/ICVFXTesting/Build/Scripts/Automation/Tests/ICVFXTest.PerformanceReport.cs b/Engine/Plugins/VirtualProduction/ICVFXTesting/Build/Scripts/Automation/Tests/ICVFXTest.PerformanceReport.cs
--- a/Engine/Plugins/VirtualProduction/ICVFXTesting/Build/Scripts/Automation/Tests/ICVFXTest.PerformanceReport.cs
+++ b/Engine/Plugins/VirtualProduction/ICVFXTesting/Build/Scripts/Automation/Tests/ICVFXTest.PerformanceReport.cs
@@ -65,6 +65,7 @@ namespace ICVFXTest
 		var CsvsPaths = new[]
 			{
 				Path.Combine(ArtifactPath, "EditorGame", "Profiling", "FPSChartStats"),
+				Path.Combine(ArtifactPath, "EditorGame", "Settings", $"{GetConfiguration().ProjectName}", "Saved", "Profiling", "FPSChartStats"),
 				Path.Combine(TempDir, "DeviceCache", Platform.ToString(), TestInstance.ClientApps[0].Device.ToString(), "UserDir")
 			};
 
diff --git a/Engine/Plugins/VirtualProduction/ICVFXTesting/Scripts/GauntletSettings.xml b/Engine/Plugins/VirtualProduction/ICVFXTesting/Scripts/GauntletSettings.xml
--- a/Engine/Plugins/VirtualProduction/ICVFXTesting/Scripts/GauntletSettings.xml
+++ b/Engine/Plugins/VirtualProduction/ICVFXTesting/Scripts/GauntletSettings.xml
@@ -72,6 +72,7 @@
 	<Property Name="ExtraGauntletCommonArgs" Value="-screenshotperiod=30 -ICVFXTest.StatCommands -ICVFXTest.MaxRunCount=$(MaxRunCount) $(LogOutputParams)" />
 	<Property Name="ExtraGauntletEditorTestArgs" Value="-game" />
 	
+	<Property Name="ExtraGauntletEditorTestArgs" Value="$(ExtraGauntletEditorTestArgs) -ICVFXTest.ProjectName=$(ProjectName)" />
 	<Property Name="ExtraGauntletEditorTestArgs" Value="$(ExtraGauntletEditorTestArgs) -ICVFXTest.TestName=$(ProjectTestName)" />
 	<Property Name="ExtraGauntletEditorTestArgs" Value="$(ExtraGauntletEditorTestArgs) -ICVFXTest.MapOverride=$(MapOverride) -ICVFXTest.DisplayConfig=$(DisplayConfigPath) -ICVFXTest.ViewportScreenPercentage=$(ViewportScreenPercentage)" />
 	<Property Name="ExtraGauntletEditorTestArgs" Value="$(ExtraGauntletEditorTestArgs) -ICVFXTest.DisplayClusterNodeName=$(DisplayClusterNodeName)" />
