Fix artifact path for SaloonPerf

diff --git a/Engine/Plugins/VirtualProduction/ICVFXTesting/Scripts/GauntletSettings.xml b/Engine/Plugins/VirtualProduction/ICVFXTesting/Scripts/GauntletSettings.xml
--- a/Engine/Plugins/VirtualProduction/ICVFXTesting/Scripts/GauntletSettings.xml
+++ b/Engine/Plugins/VirtualProduction/ICVFXTesting/Scripts/GauntletSettings.xml
@@ -63,7 +63,7 @@
 	<Property Name="SummaryReportPath" Value="$(AutomationLocalOutputDir)"/>
 	<Property Name="SummaryReportPath" Value="$(AutomationNetworkRootDir)" If="$(IsBuildMachine)"/>
 	
-	<Property Name="ReportCacheDirRoot" Value="$(AutomationLocalOutputDir)/PerfCache" />
+	<Property Name="ReportCacheDirRoot" Value="$(AutomationOutputDir)/PerfCache" />
 	
 	<Property Name="TestOutputDir" Value="$(AutomationOutputDir)/$(ProjectName)$(Platform)$(ProjectName)Perf"/>
 	<Property Name="LogOutputParams" Value="-logdir=&quot;$(TestOutputDir)&quot;"/>
