Fix RC Web interface not forwarding the IP of the request originator, which caused all requests from the Web Interface to be treated as local.

diff --git a/Engine/Plugins/VirtualProduction/RemoteControlWebInterface/WebApp/Server/src/Api.ts b/Engine/Plugins/VirtualProduction/RemoteControlWebInterface/WebApp/Server/src/Api.ts
--- a/Engine/Plugins/VirtualProduction/RemoteControlWebInterface/WebApp/Server/src/Api.ts
+++ b/Engine/Plugins/VirtualProduction/RemoteControlWebInterface/WebApp/Server/src/Api.ts
@@ -61,12 +61,12 @@ export class Api {
 
   @Put('proxy')
   private proxy(req: Request, res: Response) {
-    this.send(res, UnrealEngine.proxy(req.body?.method, req.body?.url, req.body?.body));
+    this.send(res, UnrealEngine.proxy(req.body?.method, req.body?.url, req.body?.body, req.socket.remoteAddress.replace(/^.*:/, '')));
   }
 
   @Get('thumbnail')
   private thumbnail(req: Request, res: Response) {
-    this.send(res, UnrealEngine.thumbnail(req.query.asset.toString()));
+    this.send(res, UnrealEngine.thumbnail(req.query.asset.toString(), req.socket.remoteAddress.replace(/^.*:/, '')));
   }
 
   @Get('shutdown')
diff --git a/Engine/Plugins/VirtualProduction/RemoteControlWebInterface/WebApp/Server/src/Notify.ts b/Engine/Plugins/VirtualProduction/RemoteControlWebInterface/WebApp/Server/src/Notify.ts
--- a/Engine/Plugins/VirtualProduction/RemoteControlWebInterface/WebApp/Server/src/Notify.ts
+++ b/Engine/Plugins/VirtualProduction/RemoteControlWebInterface/WebApp/Server/src/Notify.ts
@@ -11,14 +11,15 @@ export namespace Notify {
     io.attach(server);
 
     io.sockets.on('connection', (socket: socketio.Socket) => {
+      const ip = socket.handshake.address.replace(/^.*:/, '');
       socket
         .on('client', onPing)
-        .on('view', onViewChange)
-        .on('value', UnrealEngine.setPayloadValue)
-        .on('execute', UnrealEngine.executeFunction)
-        .on('metadata', UnrealEngine.setPresetPropertyMetadata)
-        .on('rebind', UnrealEngine.rebindProperties)
-        .on('search', UnrealEngine.search);
+        .on('view', (preset, view, supressUnrealNotification) => onViewChange(preset, view, supressUnrealNotification, ip))
+        .on('value', (preset, property, value) => UnrealEngine.setPayloadValue(preset, property, value, ip))
+        .on('execute', (preset, func, args) => UnrealEngine.executeFunction(preset, func, args, ip))
+        .on('metadata', (preset, property, metadata, value) => UnrealEngine.setPresetPropertyMetadata(preset, property, metadata, value, ip))
+        .on('rebind', (preset, properties, target) => UnrealEngine.rebindProperties(preset, properties, target, ip))
+        .on('search', (query, types, prefix, filterArgs, count, callback) => UnrealEngine.search(query, types, prefix, filterArgs, count, callback, ip));
 
       if (UnrealEngine.isConnected())
         socket.emit('connected', true);
@@ -34,9 +35,9 @@ export namespace Notify {
     io.emit(what, ...args);
   }
 
-  export function onViewChange(preset: string, view: IView, supressUnrealNotification?: boolean) {
+  export function onViewChange(preset: string, view: IView, supressUnrealNotification: boolean, ip: string) {
     if (!supressUnrealNotification)
-      UnrealEngine.setView(preset, view);
+      UnrealEngine.setView(preset, view, ip);
     io.emit('view', preset, view);
   }
 
diff --git a/Engine/Plugins/VirtualProduction/RemoteControlWebInterface/WebApp/Server/src/UnrealEngine.ts b/Engine/Plugins/VirtualProduction/RemoteControlWebInterface/WebApp/Server/src/UnrealEngine.ts
--- a/Engine/Plugins/VirtualProduction/RemoteControlWebInterface/WebApp/Server/src/UnrealEngine.ts
+++ b/Engine/Plugins/VirtualProduction/RemoteControlWebInterface/WebApp/Server/src/UnrealEngine.ts
@@ -143,7 +143,7 @@ export namespace UnrealEngine {
 
   export async function checkPassphrase(passphrase: string) : Promise<boolean> {
     try {
-      const res = await get<UnrealApi.PassphraseCheck>(`/remote/passphrase/`, passphrase);
+      const res = await get<UnrealApi.PassphraseCheck>(`/remote/passphrase/`, undefined, passphrase);
       if (!res?.keyCorrect)
         return false;
 
@@ -180,7 +180,6 @@ export namespace UnrealEngine {
     
     clearQuitTimeout();
     
-    console.log('Connected to UE Remote WebSocket');
     Notify.emit('connected', true);
     
     isOpen = await checkPassphrase('');
@@ -201,7 +200,7 @@ export namespace UnrealEngine {
 
   async function pullTimer() {
     try {
-      const { Presets } = await get<UnrealApi.Presets>('/remote/presets');
+      const { Presets } = await get<UnrealApi.Presets>('/remote/presets', undefined);
 
       // First check if there are any new presets
       for (const preset of Presets) {
@@ -301,7 +300,7 @@ export namespace UnrealEngine {
 
         case UnrealApi.PresetEvent.MetadataModified: {
           await refreshPreset(message.PresetId, message.PresetName);
-          refreshView(message.PresetId, message.Metadata.view);
+          refreshView(message.PresetId, message.Metadata.view, undefined);
           break;
         }
 
@@ -391,20 +390,21 @@ export namespace UnrealEngine {
       throw new Error('Websocket is not connected');
   }
 
-  function send(message: string, parameters: any, passphrase?: string) {
+  function send(message: string, parameters: any, passphrase?: string, ip?: string) {
     verifyConnection();
     const Id = parameters?.RequestId ?? wsRequest++;
     passphrase = passphrase ?? workingPassphrases[0];
-    connection.send(JSON.stringify({ MessageName: message, Id, Passphrase: passphrase, Parameters: parameters}));
+	  ip = ip ?? "";
+    connection.send(JSON.stringify({ MessageName: message, Id, Passphrase: passphrase, ForwardedFor: ip, Parameters: parameters}));
   }
 
-  function http<T>(Verb: string, URL: string, Body?: object, wantAnswer?: boolean, passphrase?: string): Promise<T> {
+  function http<T>(Verb: string, URL: string, Body?: object, wantAnswer?: boolean, passphrase?: string, ip?: string): Promise<T> {
     const RequestId = httpRequest++;
     const payload = { RequestId, Verb, URL, Body };
     if (Program.logger)
       LogServer.logLoopback(payload);
 
-    send('http', payload, passphrase);
+    send('http', payload, passphrase, ip);
     if (!wantAnswer)
       return Promise.resolve(null);
 
@@ -413,12 +413,12 @@ export namespace UnrealEngine {
     });
   }
 
-  function get<T>(url: string, passphrase?: string): Promise<T> {
-    return http<T>('GET', url, undefined, true, passphrase);
+  function get<T>(url: string, ip: string, passphrase?: string): Promise<T> {
+    return http<T>('GET', url, undefined, true, passphrase, ip);
   }
 
-  function put<T>(url: string, body: object, passphrase?: string): Promise<T> {
-    return http<T>('PUT', url, body, true, passphrase);
+  function put<T>(url: string, body: object, ip: string, passphrase?: string): Promise<T> {
+    return http<T>('PUT', url, body, true, passphrase, ip);
   }
 
   function registerToPreset(PresetName: string): void {
@@ -445,7 +445,7 @@ export namespace UnrealEngine {
     setLoading(true);
 
     try { 
-      const { Presets } = await get<UnrealApi.Presets>('/remote/presets');
+      const { Presets } = await get<UnrealApi.Presets>('/remote/presets', undefined);
 
       for (const preset of Presets ?? []) {
         const stub = preset as IPreset;
@@ -467,14 +467,14 @@ export namespace UnrealEngine {
     setLoading(false);
   }
 
-  async function pullPreset(id: string, name?: string): Promise<IPreset> {
+  async function pullPreset(id: string, name?: string, ip?: string): Promise<IPreset> {
     try {
       if (!presets[id]) {
-        const res = await get<UnrealApi.View>(`/remote/preset/${id}/metadata/view`);
-        refreshView(id, res?.Value);
+        const res = await get<UnrealApi.View>(`/remote/preset/${id}/metadata/view`, ip);
+        refreshView(id, res?.Value, ip);
       }
 
-      const { Preset } = await get<UnrealApi.Preset>(`/remote/preset/${id}`);
+      const { Preset } = await get<UnrealApi.Preset>(`/remote/preset/${id}`, ip);
       if (!Preset)
         return null;
 
@@ -485,8 +485,8 @@ export namespace UnrealEngine {
     }
   }
 
-  async function refreshPreset(id: string, name?: string): Promise<IPreset> {
-    const preset = await pullPreset(id, name);
+  async function refreshPreset(id: string, name?: string, ip?: string): Promise<IPreset> {
+    const preset = await pullPreset(id, name, ip);
     if (!preset)
       return;
 
@@ -538,20 +538,20 @@ export namespace UnrealEngine {
         preset.Exposed[Controller.ID] = Controller;
         Controllers.ExposedProperties.push(Controller);
       }
-    } 
+    }
 
     preset.IsFavorite = favorites?.[preset.ID];
     return preset;
   }
 
-  function refreshView(id: string, viewJson: string) {
+  function refreshView(id: string, viewJson: string, ip: string) {
     try {
       if (!viewJson) {
         views[id] = {
           tabs: [{ name: 'Tab 1', layout: TabLayout.Empty, icon: '' }]
         };
 
-        Notify.onViewChange(id, views[id], true);
+        Notify.onViewChange(id, views[id], true, ip);
         return;
       }
 
@@ -568,7 +568,7 @@ export namespace UnrealEngine {
       }
 
       views[id] = view;
-      Notify.onViewChange(id, view, true);
+      Notify.onViewChange(id, view, true, ip);
     } catch (error) {
       console.log('Failed to parse View of Preset', id);
     }
@@ -667,7 +667,7 @@ export namespace UnrealEngine {
     });
   }
 
-  async function pullPresetValues(Preset: IPreset): Promise<IPayload> {
+  async function pullPresetValues(Preset: IPreset, ip?: string): Promise<IPayload> {
     const Requests = [];
     for (const Property of Preset.ExposedProperties)
       Requests.push(getPropertyValueRequest(Preset, Property, false));
@@ -678,7 +678,7 @@ export namespace UnrealEngine {
     }
 
     const updatedPayloads: IPayloads = {};
-    const answers = await put<UnrealApi.BatchResponses>('/remote/batch', { Requests });
+    const answers = await put<UnrealApi.BatchResponses>('/remote/batch', { Requests }, ip);
     for (let i = 0; i < Requests.length; i++) {
       const { PropertyId, IsController } = Requests[i];
       try {
@@ -700,21 +700,21 @@ export namespace UnrealEngine {
     return updatedPayloads[Preset.ID];
   }
 
-  export async function loadPreset(id: string): Promise<IPreset> {
+  export async function loadPreset(id: string, ip?: string): Promise<IPreset> {
     if (registered[id] && presets[id])
       return presets[id];
 
     registered[id] = true;
     registerToPreset(id);
-    const preset = await refreshPreset(id);
+    const preset = await refreshPreset(id, undefined, ip);
     if (!preset)
       return;
 
-    payloads[id] = await pullPresetValues(preset);
+    payloads[id] = await pullPresetValues(preset, ip);
     Notify.emit('payloads', payloads);
     
-    const res = await get<UnrealApi.View>(`/remote/preset/${id}/metadata/view`);
-    refreshView(id, res?.Value);
+    const res = await get<UnrealApi.View>(`/remote/preset/${id}/metadata/view`, ip);
+    refreshView(id, res?.Value, ip);
   }
 
   export async function getPayload(preset: string): Promise<IPayload> {
@@ -729,7 +729,7 @@ export namespace UnrealEngine {
     payloads[preset] = payload;
   }
 
-  export async function setPayloadValue(preset: string, property: string, value: PropertyValue): Promise<void> {
+  export async function setPayloadValue(preset: string, property: string, value: PropertyValue, ip: string): Promise<void> {
     try {
       const body: any = { GenerateTransaction: true };
       if (value !== null) {
@@ -753,10 +753,10 @@ export namespace UnrealEngine {
         LogServer.logLoopback(log);
       }
 
-      await put(`/remote/preset/${preset}/property/${property}`, body);
+      await put(`/remote/preset/${preset}/property/${property}`, body, ip);
 
       if (value === null) {
-        const ret = await get<UnrealApi.PropertyValues>(`/remote/preset/${preset}/property/${property}`);
+        const ret = await get<UnrealApi.PropertyValues>(`/remote/preset/${preset}/property/${property}`, ip);
         value = ret.PropertyValues?.[0]?.PropertyValue;
         if (value !== undefined) {
           setPayloadValueInternal(payloads, [preset, property], value);
@@ -784,20 +784,20 @@ export namespace UnrealEngine {
     element[ _.last(path) ] = value;
   }
 
-  export async function executeFunction(preset: string, func: string, args: Record<string, any>): Promise<void> {
+  export async function executeFunction(preset: string, func: string, args: Record<string, any>, ip: string): Promise<void> {
     try {
       const url = `/remote/preset/${preset}/function/${func}`;
-      await put(url, { Parameters: args, GenerateTransaction: true });
+      await put(url, { Parameters: args, GenerateTransaction: true }, ip);
     } catch (err) {
       console.log('Failed to set execute function call:', err.message);
     }
   }
 
-  export async function setPresetPropertyMetadata(preset: string, property: string, metadata: string, value: string) {
+  export async function setPresetPropertyMetadata(preset: string, property: string, metadata: string, value: string, ip: string) {
     try {
 
       const url = `/remote/preset/${preset}/property/${property}/metadata/${metadata}`;
-      await put(url, { value });
+      await put(url, { value }, ip);
       const prop = presets[preset]?.Exposed[property];
       if (prop) {
         prop.Metadata[metadata] = value;
@@ -808,7 +808,7 @@ export namespace UnrealEngine {
     }
   }
 
-  export async function rebindProperties(preset: string, properties: string[], target: string) {
+  export async function rebindProperties(preset: string, properties: string[], target: string, ip: string) {
     const body = {
       objectPath: '/Script/RemoteControlWebInterface.Default__RCWebInterfaceBlueprintLibrary',
       functionName: 'RebindProperties',
@@ -819,12 +819,12 @@ export namespace UnrealEngine {
       }
     };
 
-    await put('/remote/object/call', body);
+    await put('/remote/object/call', body, ip);
 
     // We need to fetch values of new owner
     const changes: any = {};
     for (const property of properties) {
-      const ret = await get<UnrealApi.PropertyValues>(`/remote/preset/${preset}/property/${property}`);
+      const ret = await get<UnrealApi.PropertyValues>(`/remote/preset/${preset}/property/${property}`, ip);
       const value = ret.PropertyValues?.[0]?.PropertyValue;
       if (value !== undefined) {
         setPayloadValueInternal(payloads, [preset, property], value);
@@ -854,7 +854,7 @@ export namespace UnrealEngine {
     return views[preset];
   }
 
-  export async function setView(preset: string, view: IView): Promise<void> {
+  export async function setView(preset: string, view: IView, ip: string): Promise<void> {
     for (const tab of view.tabs) {
       if (!tab.panels)
         continue;
@@ -865,10 +865,10 @@ export namespace UnrealEngine {
 
     views[preset] = view;
     const Value = JSON.stringify(view);
-    await put(`/remote/preset/${preset}/metadata/view`, { Value });
+    await put(`/remote/preset/${preset}/metadata/view`, { Value }, ip);
   }
 
-  export async function search(query: string, types: string[], prefix: string, filterArgs: any, count: number, callback: (assets: IAsset[]) => void): Promise<void> {
+  export async function search(query: string, types: string[], prefix: string, filterArgs: any, count: number, callback: (assets: IAsset[]) => void, ip: string): Promise<void> {
     const ret = await put<UnrealApi.Assets>('/remote/search/assets', {
       Query: query,
       Limit: count,
@@ -880,30 +880,30 @@ export namespace UnrealEngine {
         RecursiveClasses: true,
         ...filterArgs,
       }
-    });
+    }, ip);
 
     callback?.(ret.Assets);
   }
 
-  export function proxy(method: 'GET' | 'PUT', url: string, body?: any): Promise<any> {
+  export function proxy(method: 'GET' | 'PUT', url: string, body: any, ip: string): Promise<any> {
     if (!method || !url)
       return Promise.resolve({});
 
     switch (method) {
       case 'GET':
-        return get(url);
+        return get(url, ip);
 
       case 'PUT':
         if (body)
-          return put(url, body);
+          return put(url, body, ip);
         break;
     }
 
     return Promise.resolve({});
   }
 
-  export function thumbnail(asset: string): Promise<any> {
-    return put<string>('/remote/object/thumbnail', { ObjectPath: asset })
+  export function thumbnail(asset: string, ip: string): Promise<any> {
+    return put<string>('/remote/object/thumbnail', { ObjectPath: asset }, ip)
             .then(base64 => Buffer.from(base64, 'base64'))
             .catch(error => null);
   }
diff --git a/Engine/Plugins/VirtualProduction/RemoteControlWebInterface/WebApp/_package.json b/Engine/Plugins/VirtualProduction/RemoteControlWebInterface/WebApp/_package.json
--- a/Engine/Plugins/VirtualProduction/RemoteControlWebInterface/WebApp/_package.json
+++ b/Engine/Plugins/VirtualProduction/RemoteControlWebInterface/WebApp/_package.json
@@ -1,6 +1,6 @@
 {
   "name": "conductor",
-  "version": "0.3.86",
+  "version": "0.3.87",
   "scripts": {
     "postinstall": "concurrently \"cd ./Client && npm install && cd ..\" \"cd ./Server && npm install && cd ..\"",
     "build": "node ./Scripts/build.js",
