SaloonPerf
Load Switchboard config to find display cluster config
Add TestLocation actor: The test will automatically move the display cluster to these locations.

diff --git a/Engine/Plugins/VirtualProduction/ICVFXTesting/Build/Scripts/Automation/ICVFX.TestNode.cs b/Engine/Plugins/VirtualProduction/ICVFXTesting/Build/Scripts/Automation/ICVFX.TestNode.cs
--- a/Engine/Plugins/VirtualProduction/ICVFXTesting/Build/Scripts/Automation/ICVFX.TestNode.cs
+++ b/Engine/Plugins/VirtualProduction/ICVFXTesting/Build/Scripts/Automation/ICVFX.TestNode.cs
@@ -10,13 +10,6 @@ namespace ICVFXTest
 {
 	public class ICVFXTestConfig : UnrealTestConfiguration
 	{
-		/// <summary>
-		/// Name of the project.
-		/// </summary>
-		[AutoParamWithNames("ICVFXTest.ProjectName")]
-		public string ProjectName;
-
-		/// <summary>
 		/// Name of the test.
 		/// </summary>
 		[AutoParamWithNames("ICVFXTest.TestName")]
@@ -72,13 +65,13 @@ namespace ICVFXTest
 		/// <summary>
 		/// Path to the exported ndisplay config to use for testing.
 		/// </summary>
-		[AutoParamWithNames("ICVFXTest.DisplayConfig")]
+		[AutoParamWithNames("", "ICVFXTest.DisplayConfig")]
 		public string DisplayConfigPath;
 
 		/// <summary>
 		/// Name of the display cluster node to use.
 		/// </summary>
-		[AutoParamWithNames("ICVFXTest.DisplayClusterNodeName")]
+		[AutoParamWithNames("", "ICVFXTest.DisplayClusterNodeName")]
 		public string DisplayClusterNodeName;
 
 		/// <summary>
diff --git a/Engine/Plugins/VirtualProduction/ICVFXTesting/Build/Scripts/Automation/Tests/ICVFXTest.AutoTest.cs b/Engine/Plugins/VirtualProduction/ICVFXTesting/Build/Scripts/Automation/Tests/ICVFXTest.AutoTest.cs
--- a/Engine/Plugins/VirtualProduction/ICVFXTesting/Build/Scripts/Automation/Tests/ICVFXTest.AutoTest.cs
+++ b/Engine/Plugins/VirtualProduction/ICVFXTesting/Build/Scripts/Automation/Tests/ICVFXTest.AutoTest.cs
@@ -2,6 +2,7 @@
 
 using EpicGame;
 using Gauntlet;
+using Microsoft.Extensions.Logging;
 using System.Collections.Generic;
 using System.Diagnostics;
 using UnrealBuildTool;
@@ -15,6 +16,8 @@ namespace ICVFXTest
 	/// </summary>
 	public class AutoTest : ICVFXTestNode
 	{
+		private static ILogger Logger => Log.Logger;
+		
 		public AutoTest(Gauntlet.UnrealTestContext InContext)
 			: base(InContext)
 		{
@@ -43,6 +46,26 @@ namespace ICVFXTest
 			return base.GetConfiguration().Nanite;
 		}
 
+		public ICVFXTestConfig GetBaseConfiguration()
+		{
+			return base.GetConfiguration();
+		}
+
+		public virtual string GetDisplayConfigPath()
+		{
+			return base.GetConfiguration().DisplayConfigPath;
+		}
+
+		public virtual string GetDisplayClusterNode()
+		{
+			return base.GetConfiguration().DisplayClusterNodeName;
+		}
+
+		public virtual string GetDisplayClusterUAssetPath(in string NDisplayJsonFile)
+		{
+			return "";
+		}
+
 		public override ICVFXTestConfig GetConfiguration()
 		{
 			ICVFXTestConfig Config = base.GetConfiguration();
@@ -107,7 +130,7 @@ namespace ICVFXTest
 			if (Config.MapOverride?.Length != 0)
 			{
 				ClientRole.MapOverride = Config.MapOverride;
-				Log.TraceLog($"Map Override: { Config.MapOverride }");
+				Logger.LogInformation($"Map Override: { Config.MapOverride }");
 			}
 
 			if (Config.D3DDebug)
@@ -122,14 +145,35 @@ namespace ICVFXTest
 				Gauntlet.Log.Info("Running with GPUCrashDebugging");
 			}
 
-			if (Config.DisplayConfigPath?.Length != 0)
+			string DisplayConfigPath = Config.DisplayConfigPath;
+			if (string.IsNullOrEmpty(DisplayConfigPath))
+			{
+				DisplayConfigPath = GetDisplayConfigPath();
+			}
+
+			if (DisplayConfigPath?.Length != 0)
+			{
+				string DisplayClusterUObjectPath = GetDisplayClusterUAssetPath(DisplayConfigPath);
+
+				Logger.LogInformation($"Using -> {DisplayClusterUObjectPath}.uasset");
+				if (!string.IsNullOrEmpty(DisplayClusterUObjectPath))
+				{
+					// Needed for the test runner to figure out which display cluster root actor to move in the scene.
+					ClientRole.CommandLineParams.AddOrAppendParamValue("ICVFXTest.DisplayClusterUAssetPath", DisplayClusterUObjectPath);
+				}
+
+				ClientRole.CommandLineParams.AddOrAppendParamValue("dc_cfg", DisplayConfigPath);
+			}
+
+			string DisplayConfigNode = Config.DisplayClusterNodeName;
+			if (string.IsNullOrEmpty(DisplayConfigNode))
 			{
-				ClientRole.CommandLineParams.AddOrAppendParamValue("dc_cfg", Config.DisplayConfigPath);
+				DisplayConfigNode = GetDisplayClusterNode();
 			}
 
-			if (Config.DisplayClusterNodeName?.Length != 0)
+			if (DisplayConfigNode?.Length != 0)
 			{
-				ClientRole.CommandLineParams.AddOrAppendParamValue("dc_node", Config.DisplayClusterNodeName);
+				ClientRole.CommandLineParams.AddOrAppendParamValue("dc_node", DisplayConfigNode);
 			}
 
 			if (UseVulkan())
diff --git a/Engine/Plugins/VirtualProduction/ICVFXTesting/Build/Scripts/Automation/Tests/ICVFXTest.PerformanceReport.cs b/Engine/Plugins/VirtualProduction/ICVFXTesting/Build/Scripts/Automation/Tests/ICVFXTest.PerformanceReport.cs
--- a/Engine/Plugins/VirtualProduction/ICVFXTesting/Build/Scripts/Automation/Tests/ICVFXTest.PerformanceReport.cs
+++ b/Engine/Plugins/VirtualProduction/ICVFXTesting/Build/Scripts/Automation/Tests/ICVFXTest.PerformanceReport.cs
@@ -3,24 +3,98 @@
 using AutomationTool;
 using EpicGames.Core;
 using Gauntlet;
+using ICVFXTest.Switchboard;
+using Microsoft.Extensions.Logging;
 using System;
 using System.Collections.Generic;
+using System.Diagnostics;
 using System.IO;
 using System.Linq;
+using System.Text.Json;
+using System.Text.Json.Serialization;
 using UnrealBuildBase;
 using UnrealBuildTool;
 using Log = EpicGames.Core.Log;
 
 namespace ICVFXTest
 {
+	namespace NDisplay
+	{
+		public class NDisplayInnerObject
+		{
+			[JsonPropertyName("assetPath")]
+			public string AssetPath { get; set; }
+		}
+
+		public class NDisplayConfig
+		{
+			[JsonPropertyName("nDisplay")]
+			public NDisplayInnerObject InnerConfig { get; set; }
+		}
+	}
+
+	namespace Switchboard
+	{
+		public class NDisplayDeviceSetting
+		{
+			[JsonPropertyName("ndisplay_cfg_file")]
+			public string DisplayConfigFile { get; set; }
+
+			[JsonPropertyName("primary_device_name")]
+			public string PrimaryDeviceName { get; set; }
+		}
+
+		public class NdisplaySettings
+		{
+			[JsonPropertyName("settings")]
+			public NDisplayDeviceSetting Settings { get; set; }
+		}
+
+		public class NDisplayConfig
+		{
+			[JsonPropertyName("nDisplay")]
+			public NdisplaySettings DeviceSettings { get; set; }
+		}
+
+		/// <summary>
+		/// Represents a config output by switchboard.
+		/// </summary>
+		public class SwitchboardConfig
+		{
+			public string ConfigName { get; set; }
+
+			[JsonPropertyName("project_name")]
+			public string ProjectName { get; set; }
+
+			[JsonPropertyName("uproject")]
+			public string ProjectPath { get; set; }
+
+			[JsonPropertyName("devices")]
+			public NDisplayConfig Devices { get; set; }
+		}
+
+
+		public class UserSettings
+		{
+			[JsonPropertyName("config")]
+			public string LastUsedConfig { get; set; }
+
+			[JsonPropertyName("last_browsed_path")]
+			public string LastBrowsedPath { get; set; }
+		}
+	}
+
 	/// <summary>
 	/// CI testing
 	/// </summary>
 	public class PerformanceReport : AutoTest
 	{
+		private static ILogger Logger => Log.Logger;
+		
 		public PerformanceReport(UnrealTestContext InContext)
 			: base(InContext)
 		{
+			FindSwitchboardConfigs();
 		}
 
 		public override ICVFXTestConfig GetConfiguration()
@@ -28,8 +102,9 @@ namespace ICVFXTest
 			ICVFXTestConfig Config = base.GetConfiguration();
 			UnrealTestRole ClientRole = Config.RequireRole(UnrealTargetRole.Client);
 			ClientRole.CommandLineParams.Add("csvGpuStats");
-			ClientRole.CommandLineParams.Add("csvMetadata", $"\"testname=${Config.TestName}\"");
+			ClientRole.CommandLineParams.Add("csvMetadata", $"\"testname={Config.TestName}\"");
 			ClientRole.CommandLineParams.AddOrAppendParamValue("execcmds", "t.FPSChart.DoCSVProfile 1");
+			ClientRole.CommandLineParams.AddOrAppendParamValue("execcmds", "t.FPSChart.OpenFolderOnDump 0");
 			ClientRole.CommandLineParams.Add("ICVFXTest.FPSChart");
 
 			return Config;
@@ -40,6 +115,100 @@ namespace ICVFXTest
 			base.InitHandledErrors();
 		}
 
+		public override string GetDisplayClusterUAssetPath(in string NDisplayJsonFile)
+		{
+			var NDisplayJsonFullPath = Path.Combine(Unreal.RootDirectory.FullName, NDisplayJsonFile);
+			string Text = File.ReadAllText(NDisplayJsonFullPath);
+			NDisplay.NDisplayConfig DisplayConfig = JsonSerializer.Deserialize<NDisplay.NDisplayConfig>(Text);
+			return DisplayConfig.InnerConfig.AssetPath;
+		}
+
+		void FindSwitchboardConfigs()
+		{
+			var SwitchboardConfigsFolder = Path.Combine(Unreal.RootDirectory.FullName, "Engine", "Plugins", "VirtualProduction", "Switchboard", "Source", "Switchboard", "configs");
+			if (Directory.Exists(SwitchboardConfigsFolder))
+			{
+				string[] FileEntries = Directory.GetFiles(SwitchboardConfigsFolder);
+
+				List<SwitchboardConfig> ConfigsForProject = new List<SwitchboardConfig>();
+
+				string LastConfigName = "";
+
+				foreach (string ConfigFile in FileEntries)
+				{
+					string Text = File.ReadAllText(ConfigFile);
+
+					if (ConfigFile.EndsWith("user_settings.json"))
+					{
+						var UserConfig = JsonSerializer.Deserialize<Switchboard.UserSettings>(Text);
+						LastConfigName = UserConfig.LastUsedConfig;
+						Logger.LogInformation($"Last config used: {LastConfigName}");
+						continue;
+					}
+
+					var Config = JsonSerializer.Deserialize<Switchboard.SwitchboardConfig>(Text);
+					Config.ConfigName = Path.GetFileName(ConfigFile);
+
+					if (Config.ProjectName == Context.Options.Project)
+					{
+						ConfigsForProject.Add(Config);
+					}
+				}
+				
+				Logger.LogInformation($"Detected {ConfigsForProject.Count} switchboard configs for this project.");
+
+				foreach (SwitchboardConfig Config in ConfigsForProject)
+				{
+					Logger.LogInformation($"\t{Config.ConfigName}");
+
+					if (!String.IsNullOrEmpty(LastConfigName) && Config.ConfigName == LastConfigName)
+					{
+						Logger.LogInformation($"Last NDisplay Config used: {Config.Devices.DeviceSettings.Settings.DisplayConfigFile}");
+						Logger.LogInformation($"Found DisplayConfigFile {Config.Devices.DeviceSettings.Settings.DisplayConfigFile}");
+						OverrideDisplayConfigPath = Config.Devices.DeviceSettings.Settings.DisplayConfigFile;
+
+						Logger.LogInformation($"Using autodetected DisplayConfig Node: {Config.Devices.DeviceSettings.Settings.PrimaryDeviceName}");
+						OverrideDisplayClusterNode = Config.Devices.DeviceSettings.Settings.PrimaryDeviceName;
+						break;
+					}
+				}
+
+				if ((string.IsNullOrEmpty(OverrideDisplayConfigPath) || string.IsNullOrEmpty(OverrideDisplayClusterNode)) && ConfigsForProject.Count() != 0)
+				{
+					// Fallback to first config that we found from switchboard.
+					SwitchboardConfig FirstConfig = ConfigsForProject.First();
+
+					if (string.IsNullOrEmpty(OverrideDisplayConfigPath))
+					{
+						Logger.LogInformation($"Using autodetected DisplayConfigFile {FirstConfig.Devices.DeviceSettings.Settings.DisplayConfigFile}");
+						OverrideDisplayConfigPath = FirstConfig.Devices.DeviceSettings.Settings.DisplayConfigFile;
+					}
+
+					if (string.IsNullOrEmpty(OverrideDisplayClusterNode))
+					{
+						Logger.LogInformation($"Using autodetected DisplayConfig Node: {FirstConfig.Devices.DeviceSettings.Settings.PrimaryDeviceName}");
+						OverrideDisplayClusterNode = FirstConfig.Devices.DeviceSettings.Settings.PrimaryDeviceName;
+					}
+				}
+			}
+
+			if (string.IsNullOrEmpty(OverrideDisplayClusterNode))
+			{
+				// Fallback to node_0 if we don't find anything.
+				OverrideDisplayClusterNode = "node_0";
+			}
+		}
+
+		public override string GetDisplayConfigPath()
+		{
+			return OverrideDisplayConfigPath;
+		}
+
+		public override string GetDisplayClusterNode()
+		{
+			return OverrideDisplayClusterNode;
+		}
+
 		/// <summary>
 		/// Produces a detailed csv report using PerfReportTool.
 		/// Also, stores perf data in the perf cache, and generates a historic report using the data the cache contains.
@@ -56,7 +225,7 @@ namespace ICVFXTest
 			var ToolPath = FileReference.Combine(Unreal.EngineDirectory, "Binaries", "DotNET", "CsvTools", "PerfreportTool.exe");
 			if (!FileReference.Exists(ToolPath))
 			{
-				Log.TraceError($"Failed to find perf report utility at this path: \"{ToolPath}\".");
+				Logger.LogError($"Failed to find perf report utility at this path: \"{ToolPath}\".");
 				return;
 			}
 			var ReportConfigDir = Path.Combine(Unreal.RootDirectory.FullName, "Engine", "Plugins", "VirtualProduction", "ICVFXTesting", "Build", "Scripts", "PerfReport");
@@ -65,7 +234,7 @@ namespace ICVFXTest
 		var CsvsPaths = new[]
 			{
 				Path.Combine(ArtifactPath, "EditorGame", "Profiling", "FPSChartStats"),
-				Path.Combine(ArtifactPath, "EditorGame", "Settings", $"{GetConfiguration().ProjectName}", "Saved", "Profiling", "FPSChartStats"),
+				Path.Combine(ArtifactPath, "EditorGame", "Settings", $"{Context.Options.Project}", "Saved", "Profiling", "FPSChartStats"),
 				Path.Combine(TempDir, "DeviceCache", Platform.ToString(), TestInstance.ClientApps[0].Device.ToString(), "UserDir")
 			};
 
@@ -84,7 +253,7 @@ namespace ICVFXTest
 
 			if (DiscoveredCsvs.Count == 0)
 			{
-				Log.TraceError($"Test completed successfully but no csv profiling results were found. Searched paths were:\r\n  {string.Join("\r\n  ", CsvsPaths.Select(s => $"\"{s}\""))}");
+				Logger.LogError($"Test completed successfully but no csv profiling results were found. Searched paths were:\r\n  {string.Join("\r\n  ", CsvsPaths.Select(s => $"\"{s}\""))}");
 				return;
 			}
 
@@ -97,9 +266,9 @@ namespace ICVFXTest
 				 select CsvFile).First();
 			var NewestDir = Path.GetDirectoryName(NewestFile);
 
-			Log.TraceInformation($"Using perf report cache directory \"{ReportCacheDir}\".");
-			Log.TraceInformation($"Using perf report output directory \"{ReportPath}\".");
-			Log.TraceInformation($"Using csv results directory \"{NewestDir}\". Generating historic perf report data...");
+			Logger.LogInformation($"Using perf report cache directory \"{ReportCacheDir}\".");
+			Logger.LogInformation($"Using perf report output directory \"{ReportPath}\".");
+			Logger.LogInformation($"Using csv results directory \"{NewestDir}\". Generating historic perf report data...");
 
 			// Make sure the cache and output directories exist
 			if (!Directory.Exists(ReportCacheDir))
@@ -107,7 +276,7 @@ namespace ICVFXTest
 				try { Directory.CreateDirectory(ReportCacheDir); }
 				catch (Exception Ex)
 				{
-					Log.TraceError($"Failed to create perf report cache directory \"{ReportCacheDir}\". {Ex}");
+					Logger.LogError($"Failed to create perf report cache directory \"{ReportCacheDir}\". {Ex}");
 					return;
 				}
 			}
@@ -116,7 +285,7 @@ namespace ICVFXTest
 				try { Directory.CreateDirectory(ReportPath); }
 				catch (Exception Ex)
 				{
-					Log.TraceError($"Failed to create perf report output directory \"{ReportPath}\". {Ex}");
+					Logger.LogError($"Failed to create perf report output directory \"{ReportPath}\". {Ex}");
 					return;
 				}
 			}
@@ -128,7 +297,7 @@ namespace ICVFXTest
 			CommandUtils.RunAndLog(ToolPath.FullName, $"-csvdir \"{NewestDir}\" -o \"{ReportPath}\" -reportxmlbasedir \"{ReportConfigDir}\" -summaryTableCache \"{ReportCacheDir}\" -searchpattern csvprofile* -metadatafilter platform=\"{PlatformNameFilter}\"", out int ErrorCode);
 			if (ErrorCode != 0)
 			{
-				Log.TraceError($"PerfReportTool returned error code \"{ErrorCode}\" while generating detailed report.");
+				Logger.LogError($"PerfReportTool returned error code \"{ErrorCode}\" while generating detailed report.");
 			}
 
 			// Now generate the all-time historic summary report
@@ -171,7 +340,7 @@ namespace ICVFXTest
 				CommandUtils.RunAndLog(ToolPath.FullName, ArgStr, out ErrorCode);
 				if (ErrorCode != 0)
 				{
-					Log.TraceError($"PerfReportTool returned error code \"{ErrorCode}\" while generating historic report.");
+					Logger.LogError($"PerfReportTool returned error code \"{ErrorCode}\" while generating historic report.");
 				}
 			}
 
@@ -202,7 +371,32 @@ namespace ICVFXTest
 				CommandUtils.RunAndLog(ToolPath.FullName, ArgStr, out ErrorCode);
 				if (ErrorCode != 0)
 				{
-					Log.TraceError($"PerfReportTool returned error code \"{ErrorCode}\" while generating historic report.");
+					Logger.LogError($"PerfReportTool returned error code \"{ErrorCode}\" while generating historic report.");
+				}
+				else if (!CommandUtils.IsBuildMachine)
+				{
+					/*
+					if (Directory.Exists(ReportPath))
+					{
+						ProcessStartInfo startInfo = new ProcessStartInfo
+						{
+							Arguments = ReportPath,
+							FileName = "explorer.exe"
+						};
+
+						Process.Start(startInfo);
+
+						if (File.Exists(ReportPath + "/index.html"))
+						{
+							ProcessStartInfo chromeInfo = new ProcessStartInfo
+							{
+								Arguments = ReportPath + "/index.html",
+								FileName = "chrome.exe"
+							};
+
+							Process.Start(chromeInfo);
+						}
+					}*/
 				}
 			}
 		}
@@ -211,16 +405,19 @@ namespace ICVFXTest
 		{
 			if (Result == TestResult.Passed)
 			{
-				Log.TraceInformation($"Generating performance reports using PerfReportTool.");
+				Logger.LogInformation($"Generating performance reports using PerfReportTool.");
 				GeneratePerfReport(Context.GetRoleContext(UnrealTargetRole.Client).Platform, ArtifactPath, Context.Options.TempDir);
 			}
 			else
 			{
-				Log.TraceWarning($"Skipping performance report generation because the perf report test failed.");
+				Logger.LogWarning($"Skipping performance report generation because the perf report test failed.");
 			}
 
 			return base.CreateReport(Result, Context, Build, Artifacts, ArtifactPath);
 		}
+
+		private string OverrideDisplayConfigPath;
+		private string OverrideDisplayClusterNode;
 	}
   
 	//
diff --git a/Engine/Plugins/VirtualProduction/ICVFXTesting/ICVFXTesting.uplugin b/Engine/Plugins/VirtualProduction/ICVFXTesting/ICVFXTesting.uplugin
--- a/Engine/Plugins/VirtualProduction/ICVFXTesting/ICVFXTesting.uplugin
+++ b/Engine/Plugins/VirtualProduction/ICVFXTesting/ICVFXTesting.uplugin
@@ -10,7 +10,7 @@
 	"DocsURL": "",
 	"MarketplaceURL": "",
 	"SupportURL": "",
-	"CanContainContent": false,
+	"CanContainContent": true,
 	"IsBetaVersion": true,
 	"Hidden": true,
 	"Installed": false,
@@ -18,6 +18,14 @@
 		{
 			"Name": "Gauntlet",
 			"Enabled": true
+		},
+		{
+			"Name": "LiveLinkCamera",
+			"Enabled": true
+		},
+		{
+			"Name": "nDisplay",
+			"Enabled": true
 		}
 	],
 	"Modules": [
diff --git a/Engine/Plugins/VirtualProduction/ICVFXTesting/Scripts/GauntletSettings.xml b/Engine/Plugins/VirtualProduction/ICVFXTesting/Scripts/GauntletSettings.xml
--- a/Engine/Plugins/VirtualProduction/ICVFXTesting/Scripts/GauntletSettings.xml
+++ b/Engine/Plugins/VirtualProduction/ICVFXTesting/Scripts/GauntletSettings.xml
@@ -33,7 +33,7 @@
 	
 	<!-- Test specific options -->
 	<Option Name="DisplayConfigPath" DefaultValue="" Description="The path to the .ndisplay file used for launching the cluster."/>
-	<Option Name="DisplayClusterNodeName" DefaultValue="node_0" Description="The name of the ndisplay node to launch."/>
+	<Option Name="DisplayClusterNodeName" DefaultValue="" Description="The name of the ndisplay node to launch."/>
 	<Option Name="ViewportScreenPercentage" DefaultValue="100.0" Description="The viewport screen percentage to use"/>
 	<Option Name="WithNanite" Restrict="true|false" DefaultValue="false" Description="Whether to use nanite when testing"/>
 	<Option Name="WithLumen" Restrict="true|false" DefaultValue="false" Description="Whether to use lumen when testing"/>
@@ -71,8 +71,8 @@
 		
 	<Property Name="ExtraGauntletCommonArgs" Value="-screenshotperiod=30 -ICVFXTest.StatCommands -ICVFXTest.MaxRunCount=$(MaxRunCount) $(LogOutputParams)" />
 	<Property Name="ExtraGauntletEditorTestArgs" Value="-game" />
-	
-	<Property Name="ExtraGauntletEditorTestArgs" Value="$(ExtraGauntletEditorTestArgs) -ICVFXTest.ProjectName=$(ProjectName)" />
+
+	<Property Name="ExtraGauntletEditorTestArgs" Value="$(ExtraGauntletEditorTestArgs) -ICVFXTest.IsBuildMachine=$(IsBuildMachine)" />
 	<Property Name="ExtraGauntletEditorTestArgs" Value="$(ExtraGauntletEditorTestArgs) -ICVFXTest.TestName=$(ProjectTestName)" />
 	<Property Name="ExtraGauntletEditorTestArgs" Value="$(ExtraGauntletEditorTestArgs) -ICVFXTest.MapOverride=$(MapOverride) -ICVFXTest.DisplayConfig=$(DisplayConfigPath) -ICVFXTest.ViewportScreenPercentage=$(ViewportScreenPercentage)" />
 	<Property Name="ExtraGauntletEditorTestArgs" Value="$(ExtraGauntletEditorTestArgs) -ICVFXTest.DisplayClusterNodeName=$(DisplayClusterNodeName)" />
diff --git a/Engine/Plugins/VirtualProduction/ICVFXTesting/Scripts/ICVFXTesting.xml b/Engine/Plugins/VirtualProduction/ICVFXTesting/Scripts/ICVFXTesting.xml
--- a/Engine/Plugins/VirtualProduction/ICVFXTesting/Scripts/ICVFXTesting.xml
+++ b/Engine/Plugins/VirtualProduction/ICVFXTesting/Scripts/ICVFXTesting.xml
@@ -35,8 +35,19 @@
 	<Property Name="ClientCompileAgentOverride" Value="EditorPerfWin64" />
 	<!-- Common Setting Overrides -->
 	
-	<Property Name="GauntletSettingsFile" Value="$(RootDir)/Engine/Plugins/VirtualProduction/ICVFXTesting/Scripts/GauntletSettings.xml" />
+	<Property Name="ICVFXTestingPath" Value="$(RootDir)/Engine/Plugins/VirtualProduction/ICVFXTesting/Scripts"/>
+	<Property Name="GauntletSettingsFile" Value="$(ICVFXTestingPath)/GauntletSettings.xml" />
 
 	<!-- This will declare an aggregate called BuildAndTest $(ProjectName) -->
-	<Include Script="$(RootDir)/Engine/Build/Graph/Tasks/BuildAndTestProject.xml" />
+	<Include Script="$(RootDir)/Engine/Build/Graph/Tasks/BuildAndTestProject.xml"/>
+	
+	<Agent Name="Test SaloonEditor Win64" Type="EditorPerfWin64">
+		<!--<Node Name="Test Jere">
+			<Command Name="RunUnreal" Arguments="-project=Saloon"/>
+		</Node>
+		<Node Name="Open Test Results Win64">
+			<Log Message="test"></Log>
+			<Command Name="RunPythonScript" Arguments="-pyscript=&quot;$(ICVFXTestingPath)/open_perf_folder.py&quot; &quot;$(SummaryReportPath)&quot;"/>
+		</Node>-->
+	</Agent>
 </BuildGraph>
\ No newline at end of file
diff --git a/Engine/Plugins/VirtualProduction/ICVFXTesting/Scripts/open_perf_folder.py b/Engine/Plugins/VirtualProduction/ICVFXTesting/Scripts/open_perf_folder.py
new file mode 100644
--- /dev/null
+++ b/Engine/Plugins/VirtualProduction/ICVFXTesting/Scripts/open_perf_folder.py
@@ -0,0 +1,11 @@
+import webbrowser, os
+import sys, getopt
+
+def main(argv):
+    print("Opening folder")
+    input_folder = argv[0]
+    input_folder = os.path.realpath(input_folder)
+    os.system("start " + os.path.join(input_folder, "HistoricReport_14Days_SummaryBase_Email.html"))
+
+if __name__ == "__main__":
+    main(sys.argv[1:])
\ No newline at end of file
diff --git a/Engine/Plugins/VirtualProduction/ICVFXTesting/Source/ICVFXTesting/ICVFXTesting.build.cs b/Engine/Plugins/VirtualProduction/ICVFXTesting/Source/ICVFXTesting/ICVFXTesting.build.cs
--- a/Engine/Plugins/VirtualProduction/ICVFXTesting/Source/ICVFXTesting/ICVFXTesting.build.cs
+++ b/Engine/Plugins/VirtualProduction/ICVFXTesting/Source/ICVFXTesting/ICVFXTesting.build.cs
@@ -17,7 +17,13 @@ public class ICVFXTesting : ModuleRules
 		
 		PrivateDependencyModuleNames.AddRange(
 			new string[] {
-            	"Engine"
+				"CinematicCamera",
+            	"DisplayCluster",
+				"Engine",
+                "LiveLink",
+                "LiveLinkComponents",
+				"LiveLinkInterface",
+				"RenderCore"
 			}
         );
 	}
diff --git a/Engine/Plugins/VirtualProduction/ICVFXTesting/Source/ICVFXTesting/Private/ICVFXTestControllerAutoTest.cpp b/Engine/Plugins/VirtualProduction/ICVFXTesting/Source/ICVFXTesting/Private/ICVFXTestControllerAutoTest.cpp
--- a/Engine/Plugins/VirtualProduction/ICVFXTesting/Source/ICVFXTesting/Private/ICVFXTestControllerAutoTest.cpp
+++ b/Engine/Plugins/VirtualProduction/ICVFXTesting/Source/ICVFXTesting/Private/ICVFXTestControllerAutoTest.cpp
@@ -2,13 +2,24 @@
 
 #include "ICVFXTestControllerAutoTest.h"
 
+#include "CineCameraActor.h"
+#include "Camera/CameraComponent.h"
+#include "DisplayClusterRootActor.h"
 #include "Engine/Engine.h"
 #include "Engine/LocalPlayer.h"
+#include "EngineUtils.h"
 #include "GameFramework/GameMode.h"
+#include "IDisplayCluster.h"
+#include "ICVFXTestLocation.h"
+#include "LiveLinkComponentController.h"
+#include "LiveLinkRole.h"
+#include "LiveLinkVirtualSubject.h"
 #include "Logging/LogVerbosity.h"
 #include "PlatformFeatures.h"
 #include "ProfilingDebugging/CsvProfiler.h"
 #include "VideoRecordingSystem.h"
+#include "Roles/LiveLinkTransformRole.h"
+#include "UObject/SoftObjectPtr.h"
 
 namespace ICVFXTest
 {
@@ -60,6 +71,19 @@ public:
 	FICVFXAutoTestState_InitialLoad(UICVFXTestControllerAutoTest* const TestController) 
 		: FICVFXAutoTestState(TestController) 
 	{
+		ENQUEUE_RENDER_COMMAND(SetInnerGPUIndex)(
+			[TestController](FRHICommandListImmediate& RHICmdList)
+			{
+#if WITH_MGPU
+				TestController->SetInnerGPUIndex(RHICmdList.GetGPUMask().GetLastIndex());
+#endif
+			}
+		);
+	}
+
+	virtual void Start(const EICVFXAutoTestState PrevState) override
+	{
+		
 	}
 
 	virtual void Tick(float TimeDelta) override
@@ -77,30 +101,6 @@ private:
 	const float MaxStateChangeWait = 30.0f;
 };
 
-// Sandbox Intro
-class FICVFXAutoTestState_Intro : public FICVFXAutoTestState
-{
-public:
-	FICVFXAutoTestState_Intro(UICVFXTestControllerAutoTest* const TestController) : FICVFXAutoTestState(TestController) {}
-
-	virtual void Start(const EICVFXAutoTestState PrevState) override
-	{
-		Super::Start(PrevState);
-
-		// If starting a new run, kick off requested captures
-
-		if (Controller->RequestsVideoCapture())
-		{
-			Controller->TryStartingVideoCapture();
-		}
-
-		if (Controller->RequestsFPSChart())
-		{
-			Controller->ConsoleCommand(TEXT("StartFPSChart"));
-		}
-	}
-};
-
 // Sandbox Soak
 class FICVFXAutoTestState_Soak : public FICVFXAutoTestState
 {
@@ -123,6 +123,15 @@ public:
 			{
 				Controller->ConsoleCommand(TEXT("StartFPSChart"));
 			}
+
+			// Try finding the display cluster root actor to move it during the next phases.
+			IDisplayCluster::Get().GetCallbacks().OnDisplayClusterStartScene().AddLambda([this]()
+			{
+				Controller->DisplayClusterActor = IDisplayCluster::Get().GetGameMgr()->GetRootActor();
+				Controller->UpdateInnerGPUIndex();
+				Controller->InitializeLiveLink();
+				Controller->UpdateTestLocations();
+			});
 		}
 
 		const float SoakTime = ICVFXTest::CVarSoakTime->GetFloat();
@@ -143,38 +152,59 @@ public:
 		const float SoakTime = ICVFXTest::CVarSoakTime->GetFloat();
 		if (SoakTime > 0.0f && GetTestStateTime() >= SoakTime)
 		{
-			Controller->SetTestState(EICVFXAutoTestState::Idle);
+			if (Controller->NumTestLocations())
+			{
+				Controller->SetTestState(EICVFXAutoTestState::TraverseTestLocations);
+			}
+			else
+			{
+				Controller->SetTestState(EICVFXAutoTestState::Finished);
+			}
 		}		
 	}
 };
 
-// Sandbox Soak
-class FICVFXAutoTestState_Idle : public FICVFXAutoTestState
+// Go through test locations and collect perf data.
+class FICVFXAutoTestState_TraverseTestLocations : public FICVFXAutoTestState
 {
 public:
-	FICVFXAutoTestState_Idle(UICVFXTestControllerAutoTest* const TestController) : FICVFXAutoTestState(TestController) {}
-
-	const float IdleTime = 60.f;
+	FICVFXAutoTestState_TraverseTestLocations(UICVFXTestControllerAutoTest* const TestController) : FICVFXAutoTestState(TestController) {}
 
 	virtual void Start(const EICVFXAutoTestState PrevState) override
 	{
 		Super::Start(PrevState);
+
+		//TRACE_BOOKMARK(TEXT(""));
+
 		if (!ICVFXTest::CVarTraceFileName.GetValueOnAnyThread().IsEmpty())
 		{
 			Controller->ConsoleCommand(*(FString(TEXT("trace.file ")) + ICVFXTest::CVarTraceFileName.GetValueOnAnyThread()));
 		}
 
-		UE_LOG(LogICVFXTest, Display, TEXT("AutoTest Idle %s: idle time started for %f seconds..."), ANSI_TO_TCHAR(__func__), IdleTime);
+		UE_LOG(LogICVFXTest, Display, TEXT("AutoTest TraverseTestLocations %s: started for %f seconds..."), ANSI_TO_TCHAR(__func__), Controller->TimePerTestLocation)
 	}
 
 	virtual void Tick(float TimeDelta) override
 	{
 		Super::Tick(TimeDelta);
 
-		if (IdleTime > 0.0f && GetTestStateTime() >= IdleTime)
+		Controller->TimeAtTestLocation += TimeDelta;
+		
+		if (GetTestStateTime() >= Controller->TimePerTestLocation * Controller->NumTestLocations())
 		{
 			Controller->SetTestState(EICVFXAutoTestState::Finished);
 		}
+		else if (Controller->TimeAtTestLocation >= Controller->TimePerTestLocation)
+		{
+			if (Controller->GetCurrentTestLocationIndex() == Controller->NumTestLocations() - 1)
+			{
+				Controller->SetTestState(EICVFXAutoTestState::Finished);
+			}
+			else
+			{
+				Controller->GoToNextTestLocation();
+			}
+		}
 	}
 };
 
@@ -362,6 +392,12 @@ FICVFXAutoTestState& UICVFXTestControllerAutoTest::SetTestState(const EICVFXAuto
 	return CurrentTestState;
 }
 
+void UICVFXTestControllerAutoTest::SetTestLocations(const TArray<AActor*> InTestLocations)
+{
+	TestLocations = InTestLocations;
+	StateTimeouts[(uint8)EICVFXAutoTestState::TraverseTestLocations] = TestLocations.Num() * TimePerTestLocation;
+}
+
 void UICVFXTestControllerAutoTest::EndICVFXTest(const int32 ExitCode /*= 0*/)
 {
 	if (CurrentState != EICVFXAutoTestState::Shutdown)
@@ -394,13 +430,13 @@ void UICVFXTestControllerAutoTest::OnInit()
 	// Initialize State Data
 	States[(uint8)EICVFXAutoTestState::InitialLoad] = new FICVFXAutoTestState_InitialLoad(this);
 	States[(uint8)EICVFXAutoTestState::Soak] = new FICVFXAutoTestState_Soak(this);
-	States[(uint8)EICVFXAutoTestState::Idle] = new FICVFXAutoTestState_Idle(this);
+	States[(uint8)EICVFXAutoTestState::TraverseTestLocations] = new FICVFXAutoTestState_TraverseTestLocations(this);
 	States[(uint8)EICVFXAutoTestState::Finished] = new FICVFXAutoTestState_Finished(this);
 	States[(uint8)EICVFXAutoTestState::Shutdown] = new FICVFXAutoTestState_Shutdown(this);
 
 	StateTimeouts[(uint8)EICVFXAutoTestState::InitialLoad] = 60.f;
 	StateTimeouts[(uint8)EICVFXAutoTestState::Soak] =  ICVFXTest::CVarSoakTime->GetFloat();
-	StateTimeouts[(uint8)EICVFXAutoTestState::Idle] = 60.0f;
+	StateTimeouts[(uint8)EICVFXAutoTestState::TraverseTestLocations] = 60.0f;
 	StateTimeouts[(uint8)EICVFXAutoTestState::Finished] = 300.f;
 	StateTimeouts[(uint8)EICVFXAutoTestState::Shutdown] = 30.f;
 
diff --git a/Engine/Plugins/VirtualProduction/ICVFXTesting/Source/ICVFXTesting/Private/ICVFXTestControllerBase.cpp b/Engine/Plugins/VirtualProduction/ICVFXTesting/Source/ICVFXTesting/Private/ICVFXTestControllerBase.cpp
--- a/Engine/Plugins/VirtualProduction/ICVFXTesting/Source/ICVFXTesting/Private/ICVFXTestControllerBase.cpp
+++ b/Engine/Plugins/VirtualProduction/ICVFXTesting/Source/ICVFXTesting/Private/ICVFXTestControllerBase.cpp
@@ -63,6 +63,11 @@ void UICVFXTestControllerBase::OnInit()
 	bRequestsMemReport = FParse::Param(FCommandLine::Get(), TEXT("ICVFXTest.MemReport"));
 	bRequestsVideoCapture = FParse::Param(FCommandLine::Get(), TEXT("ICVFXTest.VideoCapture"));
 
+	FParse::Value(FCommandLine::Get(), TEXT("ICVFXTest.DisplayClusterUAssetPath"), DisplayClusterUObjectPath);
+
+	DisplayClusterUObjectPath.RightChopInline(1);
+	UE_LOG(LogTemp, Error, TEXT("Test jere-1: %s"), *DisplayClusterUObjectPath);
+
 	MemReportTimerDelegate.BindUObject(this, &ThisClass::OnMemReportTimerExpired);
 
 	const FConsoleCommandDelegate MemReportIntervalDelegate = FConsoleCommandDelegate::CreateUObject(this, &ThisClass::OnMemReportIntervalChanged);
diff --git a/Engine/Plugins/VirtualProduction/ICVFXTesting/Source/ICVFXTesting/Public/ICVFXTestControllerAutoTest.h b/Engine/Plugins/VirtualProduction/ICVFXTesting/Source/ICVFXTesting/Public/ICVFXTestControllerAutoTest.h
--- a/Engine/Plugins/VirtualProduction/ICVFXTesting/Source/ICVFXTesting/Public/ICVFXTestControllerAutoTest.h
+++ b/Engine/Plugins/VirtualProduction/ICVFXTesting/Source/ICVFXTesting/Public/ICVFXTestControllerAutoTest.h
@@ -3,7 +3,22 @@
 
 #include "ICVFXTestControllerBase.h"
 
+#include "Cluster/IDisplayClusterClusterManager.h"
+#include "Components/DisplayClusterCameraComponent.h"
+#include "Components/DisplayClusterICVFXCameraComponent.h"
+#include "Controllers/LiveLinkTransformController.h"
+#include "DisplayClusterRootActor.h"
+#include "Game/IDisplayClusterGameManager.h"
 #include "HAL/IConsoleManager.h"
+#include "IDisplayClusterCallbacks.h"
+#include "ICVFXTestLocation.h"
+#include "IDisplayCluster.h"
+#include "LiveLinkComponentController.h"
+#include "LiveLinkPreset.h"
+#include "LiveLinkPresetTypes.h"
+#include "Roles/LiveLinkTransformRole.h"
+
+
 #include "ICVFXTestControllerAutoTest.generated.h"
 
 class ULocalPlayer;
@@ -13,7 +28,7 @@ enum class EICVFXAutoTestState : uint8
 {
 	InitialLoad,
 	Soak,
-	Idle ,
+	TraverseTestLocations ,
 	Finished,
 	Shutdown,
 
@@ -54,15 +69,153 @@ public:
 	FString GetStateName(const EICVFXAutoTestState State) const;
 	FICVFXAutoTestState& GetTestState() const;
 	FICVFXAutoTestState& SetTestState(const EICVFXAutoTestState NewState);
-
+	void SetTestLocations(const TArray<AActor*> TestLocations);
 	virtual void EndICVFXTest(const int32 ExitCode=0) override;
 
+	double GetCurrentStateTime() const
+	{
+		return GetTimeInCurrentState();
+	}
+
+	int32 GetCurrentTestLocationIndex() const
+	{
+		return CurrentTestLocationIndex;
+	}
+
+	void GoToTestLocation(int32 Index)
+	{
+		TimeAtTestLocation = 0.0;
+
+		if (DisplayClusterActor)
+		{
+#if WITH_EDITOR
+			UE_LOG(LogICVFXTest, Display, TEXT("AutoTest TraverseTestLocations: Moving to test location: %s"), *TestLocations[Index]->GetActorLabel());
+#endif
+			DisplayClusterActor->SetActorTransform(TestLocations[Index]->GetActorTransform());
+		}
+	}
+
+	void GoToNextTestLocation()
+	{
+		CurrentTestLocationIndex++;
+		GoToTestLocation(CurrentTestLocationIndex);
+	}
+
+	int32 NumTestLocations() const
+	{
+		return TestLocations.Num();
+	}
+
+	void SetInnerGPUIndex(int32 InGPUIndex)
+	{
+		InnerGPUIndex = InGPUIndex;
+	}
+
+	int32 GetInnerGPUIndex() const
+	{
+		return InnerGPUIndex.load();
+	}
+
+	void InitializeLiveLink() const
+	{
+		if (DisplayClusterActor)
+		{
+			ULiveLinkPreset* LLPreset = LoadObject<ULiveLinkPreset>(NULL, TEXT("/ICVFXTesting/AutomationPreset"));
+			check(LLPreset);
+
+			UDisplayClusterICVFXCameraComponent* Component = DisplayClusterActor->FindComponentByClass<UDisplayClusterICVFXCameraComponent>();
+			if (ACineCameraActor* CameraActor = Component->GetCameraSettingsICVFX().ExternalCameraActor.Get())
+			{
+				ULiveLinkComponentController* LiveLinkComponentController = NewObject<ULiveLinkComponentController>(CameraActor);
+				LiveLinkComponentController->bUpdateInEditor = true;
+				LiveLinkComponentController->bEvaluateLiveLink = true;
+
+
+				CameraActor->AddInstanceComponent(LiveLinkComponentController);
+				LiveLinkComponentController->RegisterComponentWithWorld(CameraActor->GetWorld());
+
+				FName LLVirtualSubjectName;
+				if (LLPreset->GetSubjectPresets().Num())
+				{
+					LLVirtualSubjectName = LLPreset->GetSubjectPresets()[0].Key.SubjectName.Name;
+				}
+
+				LLPreset->ApplyToClientLatent([LiveLinkComponentController, LLVirtualSubjectName, CameraActor](bool bSuccess)
+				{
+					if (bSuccess)
+					{
+						FLiveLinkSubjectRepresentation SubjectRepresentation;
+						SubjectRepresentation.Role = ULiveLinkTransformRole::StaticClass();
+						SubjectRepresentation.Subject = LLVirtualSubjectName;
+
+						LiveLinkComponentController->SetSubjectRepresentation(SubjectRepresentation);
+						LiveLinkComponentController->SetControlledComponent(ULiveLinkTransformRole::StaticClass(), CameraActor->GetCineCameraComponent());
+
+						if (TObjectPtr<ULiveLinkControllerBase>* Controller = LiveLinkComponentController->ControllerMap.Find(ULiveLinkTransformRole::StaticClass()))
+						{
+							if (ULiveLinkTransformController* TransformController = Cast<ULiveLinkTransformController>(*Controller))
+							{
+								TransformController->TransformData.bUseLocation = false;
+							}
+						}
+					}
+				});
+			}
+		}
+	}
+
+	void UpdateInnerGPUIndex()
+	{
+		TArray<UDisplayClusterICVFXCameraComponent*> CameraComponents;
+		ADisplayClusterRootActor* RootActor = Cast<ADisplayClusterRootActor>(DisplayClusterActor);
+		constexpr bool bIncludeChildActors = false;
+		RootActor->GetComponents<UDisplayClusterICVFXCameraComponent>(CameraComponents, bIncludeChildActors);
+
+		for (UDisplayClusterICVFXCameraComponent* CameraComponent : CameraComponents)
+		{
+			CameraComponent->CameraSettings.RenderSettings.AdvancedRenderSettings.GPUIndex = GetInnerGPUIndex();
+		}
+
+		const FString NodeId = IDisplayCluster::Get().GetClusterMgr()->GetNodeId();
+		RootActor->GetViewportManager()->UpdateConfiguration(EDisplayClusterRenderFrameMode::Mono, NodeId, RootActor, nullptr);
+	}
+
+	void UpdateTestLocations()
+	{
+		UClass* TestLocationClass = AICVFXTestLocation::StaticClass();
+
+		TestLocations.Reset();
+
+		for (TActorIterator<AActor> ItActor = TActorIterator<AActor>(GetWorld(), TestLocationClass); ItActor; ++ItActor)
+		{
+			TestLocations.Add(*ItActor);
+		}
+
+		if (TestLocations.Num())
+		{
+			UE_LOG(LogICVFXTest, Display, TEXT("Found %d test locations."), TestLocations.Num());
+		}
+		else
+		{
+			UE_LOG(LogICVFXTest, Display, TEXT("Could not find test locations, defaulting to display cluster position."));
+			// Todo: If no uobject path passed to us, default on the first display cluster actor we find?
+			TestLocations.Add(DisplayClusterActor);
+		}
+
+		SetTestLocations(TestLocations);
+		GoToTestLocation(0);
+	}
+
 protected:
 	virtual void OnInit() override;
 	virtual void OnPreMapChange() override;
 	virtual void OnTick(float TimeDelta) override;
 	virtual void BeginDestroy() override;
 
+public:
+	float TimePerTestLocation = 60.f;
+	double TimeAtTestLocation = 0.0;
+	AActor* DisplayClusterActor = nullptr;
 private:
 	virtual void UnbindAllDelegates() override;
 
@@ -82,4 +235,11 @@ private:
 
 	UFUNCTION()
 	void OnSoakTimeChanged();
+
+	TArray<AActor*> TestLocations;
+
+	int32 CurrentTestLocationIndex = 0;
+
+	/** What GPU Index should be used for the inner viewport. */
+	std::atomic<int32> InnerGPUIndex = 0;
 };
\ No newline at end of file
diff --git a/Engine/Plugins/VirtualProduction/ICVFXTesting/Source/ICVFXTesting/Public/ICVFXTestControllerBase.h b/Engine/Plugins/VirtualProduction/ICVFXTesting/Source/ICVFXTesting/Public/ICVFXTestControllerBase.h
--- a/Engine/Plugins/VirtualProduction/ICVFXTesting/Source/ICVFXTesting/Public/ICVFXTestControllerBase.h
+++ b/Engine/Plugins/VirtualProduction/ICVFXTesting/Source/ICVFXTesting/Public/ICVFXTestControllerBase.h
@@ -56,6 +56,9 @@ public:
 	bool TryStartingVideoCapture();
 	bool TryFinalizingVideoCapture(const bool bStopAutoContinue = false);
 
+public:
+	FString DisplayClusterUObjectPath;
+
 private:
 	float WarnStuckTime;
 	float NextWarnStuckTime;
diff --git a/Engine/Plugins/VirtualProduction/ICVFXTesting/Source/ICVFXTesting/Public/ICVFXTestLocation.h b/Engine/Plugins/VirtualProduction/ICVFXTesting/Source/ICVFXTesting/Public/ICVFXTestLocation.h
new file mode 100644
--- /dev/null
+++ b/Engine/Plugins/VirtualProduction/ICVFXTesting/Source/ICVFXTesting/Public/ICVFXTestLocation.h
@@ -0,0 +1,18 @@
+// Copyright Epic Games, Inc. All Rights Reserved.
+
+#pragma once
+
+#include "Camera/CameraActor.h"
+
+#include "ICVFXTestLocation.generated.h"
+
+/**
+ * Actor used when running ICVFX perf tests. For every test location, the display cluster root actor will be moved to this location and perf data will be connected for a minute.
+ */
+UCLASS(Blueprintable, BlueprintType, Category="ICVFX", DisplayName="ICVFXTestLocation")
+class ICVFXTESTING_API AICVFXTestLocation : public ACameraActor
+{
+public:
+    GENERATED_BODY()
+    
+};
\ No newline at end of file
