Remote Control Web Interface
Fix Passphrase login prompt not working

diff --git a/Engine/Plugins/VirtualProduction/RemoteControlWebInterface/WebApp/Client/src/App.tsx b/Engine/Plugins/VirtualProduction/RemoteControlWebInterface/WebApp/Client/src/App.tsx
--- a/Engine/Plugins/VirtualProduction/RemoteControlWebInterface/WebApp/Client/src/App.tsx
+++ b/Engine/Plugins/VirtualProduction/RemoteControlWebInterface/WebApp/Client/src/App.tsx
@@ -620,9 +620,9 @@ export class App extends React.Component<PropsFromState, State> {
 
     if (!connected)
       return <ConnectionStatus />;
-
+  
     if (!keyCorrect && !isOpen)
-      return <Login />;
+      return <Login onSubmit={this.checkPassphrase} />;
 
     if (loading) {
       return (
@@ -651,7 +651,7 @@ export class App extends React.Component<PropsFromState, State> {
       appClassName = appClassName.replace('closed', 'open');
 
     if (lockedUI)
-      appClassName += 'locked ';  
+      appClassName += 'locked ';
 
     return (
       <DragDropContext onBeforeCapture={this.onBeforeCapture}
diff --git a/Engine/Plugins/VirtualProduction/RemoteControlWebInterface/WebApp/Client/src/components/Login.tsx b/Engine/Plugins/VirtualProduction/RemoteControlWebInterface/WebApp/Client/src/components/Login.tsx
--- a/Engine/Plugins/VirtualProduction/RemoteControlWebInterface/WebApp/Client/src/components/Login.tsx
+++ b/Engine/Plugins/VirtualProduction/RemoteControlWebInterface/WebApp/Client/src/components/Login.tsx
@@ -1,10 +1,61 @@
 import React from 'react';
 import { ReactComponent as Logo } from '../assets/ue_logo.svg';
+import crypto from 'crypto';
 
+type Props = {
+  value?: string;
+  onSubmit?: (password) => void;
+}
 
-export class Login extends React.Component {
+type State = {
+  value: string;
+  error: boolean;
+};
+
+export class Login extends React.Component<Props> {
+  onClick = () => {
+    this.submit();
+  }
+
+  submit = async (initialSubmit) => {
+	const secured = crypto.createHash('md5').update(this.state.value).digest('hex');
+	try {
+		await this.props.onSubmit(secured);
+	} catch (e) {
+		const newState = { value: this.state.value, error: !initialSubmit };
+		this.setState(newState);
+	}
+  }
+
+  state: State = {
+    value: '',
+	error: false
+  };
+
+  componentDidMount() {
+    document.addEventListener('keydown', this.onKeyPress);
+    const { value } = this.props;
+    this.setState({ value });
+	this.submit(true);
+  }
+
+  onChange = (e: React.ChangeEvent<HTMLInputElement>) => {
+    const value = e.target.value;
+    this.setState({ value });
+  }
+
+  componentWillUnmount() {    
+    document.removeEventListener('keydown', this.onKeyPress);
+  }
+
+  onKeyPress = (e) => {
+    if (e.key === 'Enter')
+      this.submit();
+  }
 
   render() {
+    const { value } = this.props;
+
     return (
       <div className="fullscreen login-screen">
         <div className='icon-wrapper'>
@@ -15,9 +66,12 @@ export class Login extends React.Component {
         </div>
         <div className='form'>
           Password
-          <input type='password'/>
+          <input type='password' value={value} onChange={this.onChange}/>
+          <button className="btn btn-secondary" onClick={this.onClick}>Submit</button>
           <div className='hint'>Configure the password in: <span className="breadcrumb">Project Settings / Authentication / Password</span></div>
-          <label className='login-status'>Incorrect Password</label>
+		  
+          <label className='login-status' style={{visibility: this.state.error ? 'visible' : 'hidden' }}>Incorrect Password</label>
+		  
         </div>
      </div>
     );
diff --git a/Engine/Plugins/VirtualProduction/RemoteControlWebInterface/WebApp/_package.json b/Engine/Plugins/VirtualProduction/RemoteControlWebInterface/WebApp/_package.json
--- a/Engine/Plugins/VirtualProduction/RemoteControlWebInterface/WebApp/_package.json
+++ b/Engine/Plugins/VirtualProduction/RemoteControlWebInterface/WebApp/_package.json
@@ -1,6 +1,6 @@
 {
   "name": "conductor",
-  "version": "0.3.87",
+  "version": "0.3.88",
   "scripts": {
     "postinstall": "concurrently \"cd ./Client && npm install && cd ..\" \"cd ./Server && npm install && cd ..\"",
     "build": "node ./Scripts/build.js",
