Fix rc web interface build

diff --git a/Engine/Plugins/VirtualProduction/RemoteControlWebInterface/WebApp/Client/src/components/Login.tsx b/Engine/Plugins/VirtualProduction/RemoteControlWebInterface/WebApp/Client/src/components/Login.tsx
--- a/Engine/Plugins/VirtualProduction/RemoteControlWebInterface/WebApp/Client/src/components/Login.tsx
+++ b/Engine/Plugins/VirtualProduction/RemoteControlWebInterface/WebApp/Client/src/components/Login.tsx
@@ -17,14 +17,14 @@ export class Login extends React.Component<Props> {
     this.submit();
   }
 
-  submit = async (initialSubmit) => {
-	const secured = crypto.createHash('md5').update(this.state.value).digest('hex');
-	try {
-		await this.props.onSubmit(secured);
-	} catch (e) {
-		const newState = { value: this.state.value, error: !initialSubmit };
-		this.setState(newState);
-	}
+  submit = async (initialSubmit?: boolean) => {
+    const secured = crypto.createHash('md5').update(this.state.value).digest('hex');
+    try {
+      await this.props.onSubmit(secured);
+    } catch (e) {
+      const newState = { value: this.state.value, error: !initialSubmit };
+      this.setState(newState);
+    }
   }
 
   state: State = {
diff --git a/Engine/Plugins/VirtualProduction/RemoteControlWebInterface/WebApp/_package.json b/Engine/Plugins/VirtualProduction/RemoteControlWebInterface/WebApp/_package.json
--- a/Engine/Plugins/VirtualProduction/RemoteControlWebInterface/WebApp/_package.json
+++ b/Engine/Plugins/VirtualProduction/RemoteControlWebInterface/WebApp/_package.json
@@ -1,6 +1,6 @@
 {
   "name": "conductor",
-  "version": "0.3.88",
+  "version": "0.3.89",
   "scripts": {
     "postinstall": "concurrently \"cd ./Client && npm install && cd ..\" \"cd ./Server && npm install && cd ..\"",
     "build": "node ./Scripts/build.js",
