Fix double path separators causing exception generating PCBs.

diff --git a/Engine/Plugins/VirtualProduction/ICVFXTesting/Build/Scripts/ICVFX.Automation.csproj b/Engine/Plugins/VirtualProduction/ICVFXTesting/Build/Scripts/ICVFX.Automation.csproj
--- a/Engine/Plugins/VirtualProduction/ICVFXTesting/Build/Scripts/ICVFX.Automation.csproj
+++ b/Engine/Plugins/VirtualProduction/ICVFXTesting/Build/Scripts/ICVFX.Automation.csproj
@@ -2,7 +2,7 @@
   <Import Project="ICVFX.Automation.csproj.props" Condition="Exists('ICVFX.Automation.csproj.props')"/>
 
   <PropertyGroup>
-    <EngineDir>../../../../</EngineDir>
+    <EngineDir>../../../..</EngineDir>
   </PropertyGroup>
   
   <Import Project="$(EngineDir)\Source\Programs\Shared\UnrealEngine.csproj.props" />
