Fix race condition in web interface requests that would result in an empty preset.
Fix forwarded local IP being treated as remote IP for security purposes

diff --git a/Engine/Plugins/VirtualProduction/RemoteControl/Source/WebRemoteControl/Private/RemoteControlDefaultPreprocessors.h b/Engine/Plugins/VirtualProduction/RemoteControl/Source/WebRemoteControl/Private/RemoteControlDefaultPreprocessors.h
--- a/Engine/Plugins/VirtualProduction/RemoteControl/Source/WebRemoteControl/Private/RemoteControlDefaultPreprocessors.h
+++ b/Engine/Plugins/VirtualProduction/RemoteControl/Source/WebRemoteControl/Private/RemoteControlDefaultPreprocessors.h
@@ -334,7 +334,7 @@ namespace UE::WebRemoteControl
 						{
 							PeerAddress = ForwardedIP->Last();
 
-							if (GetDefault<URemoteControlSettings>()->IsClientAllowed(PeerAddress))
+							if (GetDefault<URemoteControlSettings>()->IsClientAllowed(PeerAddress) || IsLocal(PeerAddress))
 							{
 								return FPreprocessorResult::Passthrough();
 							}
diff --git a/Engine/Plugins/VirtualProduction/RemoteControlWebInterface/WebApp/Client/src/reducers/api.ts b/Engine/Plugins/VirtualProduction/RemoteControlWebInterface/WebApp/Client/src/reducers/api.ts
--- a/Engine/Plugins/VirtualProduction/RemoteControlWebInterface/WebApp/Client/src/reducers/api.ts
+++ b/Engine/Plugins/VirtualProduction/RemoteControlWebInterface/WebApp/Client/src/reducers/api.ts
@@ -351,11 +351,16 @@ reducer
       }
 
       _preset = preset;
-      _api.presets.load(preset)
+
+      // Delay loading the preset so it happens after we pull the initial list of presets.
+      //If it happens intertwined with or before these requests, the result will get overwritten with an empty preset since we haven't loaded it.
+      setTimeout(() => {
+        _api.presets.load(preset)
           .then(() => Promise.all([
               _api.views.get(preset),
               _api.payload.get(preset),
           ]));
+      },200);
 
       return { ...state, preset, view: { tabs: [] }, payload: {} };
     }
diff --git a/Engine/Plugins/VirtualProduction/RemoteControlWebInterface/WebApp/_package.json b/Engine/Plugins/VirtualProduction/RemoteControlWebInterface/WebApp/_package.json
--- a/Engine/Plugins/VirtualProduction/RemoteControlWebInterface/WebApp/_package.json
+++ b/Engine/Plugins/VirtualProduction/RemoteControlWebInterface/WebApp/_package.json
@@ -1,6 +1,6 @@
 {
   "name": "conductor",
-  "version": "0.3.90",
+  "version": "0.3.91",
   "scripts": {
     "postinstall": "concurrently \"cd ./Client && npm install && cd ..\" \"cd ./Server && npm install && cd ..\"",
     "build": "node ./Scripts/build.js",
