Fix warning in Saloon buildgraph

#rnx

diff --git a/Engine/Plugins/VirtualProduction/ICVFXTesting/Scripts/GauntletSettings.xml b/Engine/Plugins/VirtualProduction/ICVFXTesting/Scripts/GauntletSettings.xml
--- a/Engine/Plugins/VirtualProduction/ICVFXTesting/Scripts/GauntletSettings.xml
+++ b/Engine/Plugins/VirtualProduction/ICVFXTesting/Scripts/GauntletSettings.xml
@@ -55,8 +55,8 @@
 	<Property Name="AutomationLocalOutputDir" Value="$(RootDir)/LocalBuilds/Reports/$(ProjectName)Perf"/>
 
 	<!-- Outputs for this run keeping slashes flipped the right way because it shows up on the webapp this way -->
-	<Property Name="AutomationNetworkOutputDir" Value="$(AutomationNetworkRootDir)\$(BuildName)"/>
-	<Property Name="AutomationNetworkOutputDir" Value="$(AutomationNetworkRootDir)/$(BuildName)" If="'$(HostPlatform)' != 'Win64'"/>
+	<Property Name="AutomationNetworkOutputDir" Value="$(AutomationNetworkRootDir)\$(BuildNamePath)"/>
+	<Property Name="AutomationNetworkOutputDir" Value="$(AutomationNetworkRootDir)/$(BuildNamePath)" If="'$(HostPlatform)' != 'Win64'"/>
 	<Property Name="AutomationOutputDir" Value="$(AutomationLocalOutputDir)"/>
 	<Property Name="AutomationOutputDir" Value="$(AutomationNetworkOutputDir)" If="$(IsBuildMachine)"/>
 	
@@ -86,7 +86,7 @@
 	<Property Name="ExtraGauntletEditorTestArgs" Value="$(ExtraGauntletEditorTestArgs) -ICVFXTest.NoRHIThread" If="$(WithNoRHIThread)"/>
 	<Property Name="ExtraGauntletEditorTestArgs" Value="$(ExtraGauntletEditorTestArgs) -ICVFXTest.DisableVirtualShadowMaps" If="!$(WithVirtualShadowMaps)"/>
 	<Property Name="ExtraGauntletEditorTestArgs" Value="$(ExtraGauntletEditorTestArgs) -ICVFXTest.EnableTrace" If="$(WithTrace)"/>
-	<Property Name="ExtraGauntletEditorTestArgs" Value="$(ExtraGauntletEditorTestArgs) -ICVFXTest.TraceRootFolder=$(TestOutputDir)/$(BuildName)" If="$(WithTrace)"/>
+	<Property Name="ExtraGauntletEditorTestArgs" Value="$(ExtraGauntletEditorTestArgs) -ICVFXTest.TraceRootFolder=$(TestOutputDir)/$(BuildNamePath)" If="$(WithTrace)"/>
 	<Property Name="ExtraGauntletEditorTestArgs" Value="$(ExtraGauntletEditorTestArgs) -ICVFXTest.PerfCacheFolder=$(ReportCacheDirRoot)"/>
 	
 	<!-- Common Gauntlet settings -->
