Add map override option to SaloonBuildAndTest horde job

diff --git a/Engine/Plugins/VirtualProduction/ICVFXTesting/Scripts/ICVFXTesting.xml b/Engine/Plugins/VirtualProduction/ICVFXTesting/Scripts/ICVFXTesting.xml
--- a/Engine/Plugins/VirtualProduction/ICVFXTesting/Scripts/ICVFXTesting.xml
+++ b/Engine/Plugins/VirtualProduction/ICVFXTesting/Scripts/ICVFXTesting.xml
@@ -40,14 +40,4 @@
 
 	<!-- This will declare an aggregate called BuildAndTest $(ProjectName) -->
 	<Include Script="$(RootDir)/Engine/Build/Graph/Tasks/BuildAndTestProject.xml"/>
-	
-	<Agent Name="Test SaloonEditor Win64" Type="EditorPerfWin64">
-		<!--<Node Name="Test Jere">
-			<Command Name="RunUnreal" Arguments="-project=Saloon"/>
-		</Node>
-		<Node Name="Open Test Results Win64">
-			<Log Message="test"></Log>
-			<Command Name="RunPythonScript" Arguments="-pyscript=&quot;$(ICVFXTestingPath)/open_perf_folder.py&quot; &quot;$(SummaryReportPath)&quot;"/>
-		</Node>-->
-	</Agent>
 </BuildGraph>
\ No newline at end of file
