Fix SaloonPerf horde job using Vulkan in nanite run

diff --git a/Engine/Plugins/VirtualProduction/ICVFXTesting/Build/Scripts/Automation/Tests/ICVFXTest.PerformanceReport.cs b/Engine/Plugins/VirtualProduction/ICVFXTesting/Build/Scripts/Automation/Tests/ICVFXTest.PerformanceReport.cs
--- a/Engine/Plugins/VirtualProduction/ICVFXTesting/Build/Scripts/Automation/Tests/ICVFXTest.PerformanceReport.cs
+++ b/Engine/Plugins/VirtualProduction/ICVFXTesting/Build/Scripts/Automation/Tests/ICVFXTest.PerformanceReport.cs
@@ -514,10 +514,6 @@ namespace ICVFXTest
 			return "Nanite";
 		}
 
-		public override bool UseVulkan()
-		{
-			return true;
-		}
 		public override bool UseNanite()
 		{
 			return true;
