update platforms

#rnx

diff --git a/Engine/Plugins/VirtualProduction/ICVFXTesting/Build/Scripts/PerfReport/ReportTypes.xml b/Engine/Plugins/VirtualProduction/ICVFXTesting/Build/Scripts/PerfReport/ReportTypes.xml
--- a/Engine/Plugins/VirtualProduction/ICVFXTesting/Build/Scripts/PerfReport/ReportTypes.xml
+++ b/Engine/Plugins/VirtualProduction/ICVFXTesting/Build/Scripts/PerfReport/ReportTypes.xml
@@ -30,6 +30,7 @@
 		<mapping destName="deviceType" destValue="console" sourceName="platform" sourceValue="ps4"/>
 		<mapping destName="deviceType" destValue="console" sourceName="platform" sourceValue="ps5"/>
 		<mapping destName="deviceType" destValue="console" sourceName="platform" sourceValue="switch"/>
+		<mapping destName="deviceType" destValue="console" sourceName="platform" sourceValue="xb1"/>
 		<mapping destName="deviceType" destValue="console" sourceName="platform" sourceValue="xboxone"/>
 		<mapping destName="deviceType" destValue="console" sourceName="platform" sourceValue="xboxonegdk"/>
 		<mapping destName="deviceType" destValue="console" sourceName="platform" sourceValue="xsx"/>
