Removed unused report graphs & CitySample code
Generate a CSV report alongside the HTML report
Switched function signatures for HistoricReport(...) and HistoricReport_Alt(...)

diff --git a/Engine/Plugins/VirtualProduction/ICVFXTesting/Build/Scripts/Automation/Tests/ICVFXTest.PerformanceReport.cs b/Engine/Plugins/VirtualProduction/ICVFXTesting/Build/Scripts/Automation/Tests/ICVFXTest.PerformanceReport.cs
--- a/Engine/Plugins/VirtualProduction/ICVFXTesting/Build/Scripts/Automation/Tests/ICVFXTest.PerformanceReport.cs
+++ b/Engine/Plugins/VirtualProduction/ICVFXTesting/Build/Scripts/Automation/Tests/ICVFXTest.PerformanceReport.cs
@@ -231,25 +231,25 @@ namespace ICVFXTest
 			var ReportConfigDir = Path.Combine(Unreal.RootDirectory.FullName, "Engine", "Plugins", "VirtualProduction", "ICVFXTesting", "Build", "Scripts", "PerfReport");
 			var ReportPath = Path.Combine(ArtifactPath, "Reports", "Performance");
 
-		var CsvsPaths = new[]
-			{
-				Path.Combine(ArtifactPath, "EditorGame", "Profiling", "FPSChartStats"),
-				Path.Combine(ArtifactPath, "EditorGame", "Settings", $"{Context.Options.Project}", "Saved", "Profiling", "FPSChartStats"),
-				Path.Combine(TempDir, "DeviceCache", Platform.ToString(), TestInstance.ClientApps[0].Device.ToString(), "UserDir")
-			};
+			var CsvsPaths = new[]
+				{
+					Path.Combine(ArtifactPath, "EditorGame", "Profiling", "FPSChartStats"),
+					Path.Combine(ArtifactPath, "EditorGame", "Settings", $"{Context.Options.Project}", "Saved", "Profiling", "FPSChartStats"),
+					Path.Combine(TempDir, "DeviceCache", Platform.ToString(), TestInstance.ClientApps[0].Device.ToString(), "UserDir")
+				};
 
 
-		var DiscoveredCsvs = new List<string>();
-			foreach (var CsvsPath in CsvsPaths)
-			{
-				if (Directory.Exists(CsvsPath))
+			var DiscoveredCsvs = new List<string>();
+				foreach (var CsvsPath in CsvsPaths)
 				{
-					DiscoveredCsvs.AddRange(
-						from CsvFile in Directory.GetFiles(CsvsPath, "*.csv", SearchOption.AllDirectories)
-						where CsvFile.Contains("csvprofile", StringComparison.InvariantCultureIgnoreCase)
-						select CsvFile);
+					if (Directory.Exists(CsvsPath))
+					{
+						DiscoveredCsvs.AddRange(
+							from CsvFile in Directory.GetFiles(CsvsPath, "*.csv", SearchOption.AllDirectories)
+							where CsvFile.Contains("csvprofile", StringComparison.InvariantCultureIgnoreCase)
+							select CsvFile);
+					}
 				}
-			}
 
 			if (DiscoveredCsvs.Count == 0)
 			{
@@ -300,28 +300,10 @@ namespace ICVFXTest
 				Logger.LogError($"PerfReportTool returned error code \"{ErrorCode}\" while generating detailed report.");
 			}
 
-			// Now generate the all-time historic summary report
-			HistoricReport("HistoricReport_AllTime", new[]
-			{
-				$"platform={PlatformNameFilter}"
-			});
-
-			// 14 days historic report
-			HistoricReport($"HistoricReport_14Days", new[]
-			{
-				$"platform={PlatformNameFilter}",
-				$"starttimestamp>={DateTimeOffset.Now.ToUnixTimeSeconds() - (14 * 60L * 60L * 24L)}"
-			});
-
-			// 7 days historic report
-			HistoricReport($"HistoricReport_7Days", new[]
-			{
-				$"platform={PlatformNameFilter}",
-				$"starttimestamp>={DateTimeOffset.Now.ToUnixTimeSeconds() - (7 * 60L * 60L * 24L)}"
-			});
-
-			void HistoricReport_Alt(string Name, IEnumerable<string> Filter)
+			// Generates HTML & CSV reports to Reports/SaloonPerf
+			void HistoricReport(string Name, IEnumerable<string> Filter)
 			{
+				// Generate HTML report
 				var Args = new[]
 				{
 					$"-summarytablecachein \"{ReportCacheDir}\"",
@@ -340,19 +322,37 @@ namespace ICVFXTest
 				CommandUtils.RunAndLog(ToolPath.FullName, ArgStr, out ErrorCode);
 				if (ErrorCode != 0)
 				{
-					Logger.LogError($"PerfReportTool returned error code \"{ErrorCode}\" while generating historic report.");
+					Logger.LogError($"PerfReportTool returned error code \"{ErrorCode}\" while generating HTML for the historic report.");
 				}
-			}
 
-			// 14 days historic report
-			HistoricReport_Alt($"HistoricReport_14Days_Summary", new[]
-			{
-				$"platform={PlatformNameFilter}",
-				$"starttimestamp>={DateTimeOffset.Now.ToUnixTimeSeconds() - (14 * 60L * 60L * 24L)}"
-			});
 
-			void HistoricReport(string Name, IEnumerable<string> Filter)
-			{
+				// Generate CSV report
+				Args = new[]
+				{
+					$"-summarytablecachein \"{ReportCacheDir}\"",
+					$"-reportxmlbasedir \"{ReportConfigDir}\"",
+					$"-o \"{base.GetConfiguration().SummaryReportPath}\"",
+					$"-metadatafilter \"{string.Join(" and ", Filter)}\"",
+					$"-summaryTableFilename \"{Name  + GetTestSuffix()}\"",
+					"-csvTable",
+					"-summaryTable autoPerfReportStandard",
+					"-condensedSummaryTable autoPerfReportStandard",
+					"-emailtable",
+					"-recurse"
+				};
+
+				ArgStr = string.Join(" ", Args);
+
+				CommandUtils.RunAndLog(ToolPath.FullName, ArgStr, out ErrorCode);
+				if (ErrorCode != 0)
+				{
+					Logger.LogError($"PerfReportTool returned error code \"{ErrorCode}\" while generating CSV for the historic report.");
+				}
+			}
+
+			// Generates HTML report to Reports/SaloonPerf/SaloonWin64SaloonPerf/ICVFXTest.PerformanceReport_PLATFORM/Reports/Performance
+			void HistoricReport_Alt(string Name, IEnumerable<string> Filter)
+			{	
 				var Args = new[]
 				{
 					$"-summarytablecachein \"{ReportCacheDir}\"",
@@ -373,32 +373,27 @@ namespace ICVFXTest
 				{
 					Logger.LogError($"PerfReportTool returned error code \"{ErrorCode}\" while generating historic report.");
 				}
-				else if (!CommandUtils.IsBuildMachine)
-				{
-					/*
-					if (Directory.Exists(ReportPath))
-					{
-						ProcessStartInfo startInfo = new ProcessStartInfo
-						{
-							Arguments = ReportPath,
-							FileName = "explorer.exe"
-						};
-
-						Process.Start(startInfo);
-
-						if (File.Exists(ReportPath + "/index.html"))
-						{
-							ProcessStartInfo chromeInfo = new ProcessStartInfo
-							{
-								Arguments = ReportPath + "/index.html",
-								FileName = "chrome.exe"
-							};
-
-							Process.Start(chromeInfo);
-						}
-					}*/
-				}
 			}
+
+			// All-time historic summary report
+			HistoricReport_Alt("HistoricReport_AllTime", new[]
+			{
+				$"platform={PlatformNameFilter}"
+			});
+
+			// 14 days historic report
+			HistoricReport_Alt($"HistoricReport_14Days", new[]
+			{
+				$"platform={PlatformNameFilter}",
+				$"starttimestamp>={DateTimeOffset.Now.ToUnixTimeSeconds() - (14 * 60L * 60L * 24L)}"
+			});
+
+			// 14 days historic report
+			HistoricReport($"HistoricReport_14Days_Summary", new[]
+			{
+				$"platform={PlatformNameFilter}",
+				$"starttimestamp>={DateTimeOffset.Now.ToUnixTimeSeconds() - (14 * 60L * 60L * 24L)}"
+			});
 		}
 
 		public override ITestReport CreateReport(TestResult Result, UnrealTestContext Context, UnrealBuildSource Build, IEnumerable<UnrealRoleResult> Artifacts, string ArtifactPath)
diff --git a/Engine/Plugins/VirtualProduction/ICVFXTesting/Build/Scripts/PerfReport/ReportTypes.xml b/Engine/Plugins/VirtualProduction/ICVFXTesting/Build/Scripts/PerfReport/ReportTypes.xml
--- a/Engine/Plugins/VirtualProduction/ICVFXTesting/Build/Scripts/PerfReport/ReportTypes.xml
+++ b/Engine/Plugins/VirtualProduction/ICVFXTesting/Build/Scripts/PerfReport/ReportTypes.xml
@@ -47,26 +47,6 @@
 	</csvEventsToStrip>
 
 
-	<sharedSummaries>
-		<summary refName="mapOverlay" type="mapoverlay" useUnstrippedCsvStats="1" title="Map" framesPerLineSegment="15" 
-			width="600" height="600" destImage="Minimap.jpg" sourceImage="CitySample_Big_City_LVL_map.png" 
-			offsetX="0.0" offsetY="0.0" scale="1.0" lineSplitDistanceThreshold="500">
-			
-			<!-- TO DO - Look at tracking and drawing player states as part of movement across map ie walk, drone, drive-->
-			<!-- Note: summaryStatName/summaryStatNamePrefix is not yet supported! -->
-			<!--
-			<overlay xStat="view/posx" yStat="view/posy" zStat="view/posz" lineColor="#ff0000" summaryStatNamePrefix="view0">
-				<event name="CitySample/GamePhase Step: IntroCine" shortName="IntroCine" lineColor="#ffff00" isStartEvent="1"/>
-				<event name="FooBar begin" shortName="FooBar Begin" summaryStatName="FooBarBeginPos" lineColor="#ffa000"/>
-				<event name="FooBar end" shortName="FooBar End" summaryStatName="FooBarEndPos"/>
-			</overlay>
-			<overlay xStat="view1/posx" yStat="view1/posy" zStat="view1/posz" lineColor="#0000ff" summaryStatNamePrefix="view1">
-				<event name="FooBar begin 1" shortName="FooBar begin 1" summaryStatName="FooBarBeginPos" isStartEvent="1" lineColor="#fff311"/>
-				<event name="FooBar end 1" shortName="FooBar End 1" summaryStatName="FooBarEndPos"/>
-			</overlay>
-			-->
-		</summary>
-	</sharedSummaries>
 
 	<!-- Preset summary tables -->
 	<summaryTableColumnFormatInfo>
@@ -77,7 +57,8 @@
 		<columnInfo name="Total Time*" autoColorize="lowIsBad"/>
 	</summaryTableColumnFormatInfo>
 
-<!--
+	<!--
+	<sharedSummaries>
 		<summary refName="psos" type="event" title="PSOs" metadataKey="PSOs">
 			<events>PSO/*</events>
 			<colourThresholds>0, 1, 2, 3</colourThresholds>
@@ -132,6 +113,7 @@
 	<reporttypes reportGraphsFile="ReportGraphs.xml">
 
 		<!-- 60fps report -->
+		<!-- This report is auto-detected -->
 		<reporttype name="Default60fps" title="60FPS Performance Report" vsync="1">
 			<autodetection requiredstats="frametime">
 				<requiredmetadata key="targetframerate" allowedValues="60" stopIfKeyFound="1" ignoreIfKeyNotFound="1"/>
@@ -153,7 +135,6 @@
 				<colourThresholds stat="GPUtime"> 16.7, 17.5, 19, 22.5 </colourThresholds>
 				<colourThresholds stat="Basic/TicksQueued"> 150, 200, 250, 300 </colourThresholds>
 				<colourThresholds stat="fmsglogf/fmsglogfcount">0.01,0.1,0.2,0.5</colourThresholds>
-			<!--<capture name="SkyDiving" startEvent="BusFlying" endEvent="SafeZones" includeFirstFrame="false" includeLastFrame="false"--> 
 			</summary>
 
 			<summary type="hitches">
@@ -200,190 +181,5 @@
 			<graph title="Log Counts" budget="0" />
 			<graph title="CSV Profiler" budget="0" />
 		</reporttype>
-
-
-		<!-- 30 Fps report -->
-		<reporttype name="Default30fps" title="30FPS Performance Report" vsync="1">
-			<autodetection requiredstats="frametime">
-				<requiredmetadata key="targetframerate" allowedValues="30" stopIfKeyFound="1" ignoreIfKeyNotFound="1"/>
-			</autodetection>
-
-			<metadataToShow>buildversion,platform,deviceprofile,csvfilename,rhiname,desktophwinfo,testname</metadataToShow>
-
-			<summary type="fpschart" fps="30" hitchThreshold="60" useEngineHitchMetric="1">
-				<stats>Frametime,GameThreadtime,RenderThreadtime,RHIThreadTime,GPUtime,RHI/Drawcalls,DynamicResolutionPercentage,MemoryFreeMB(min),PhysicalUsedMB(max),Basic/TicksQueued</stats>
-				<colourThresholds stat="Hitches/Min"> 0, 2.5, 3.75, 5</colourThresholds>
-				<colourThresholds stat="HitchTimePercent"> 0.1, 0.5, 1, 2</colourThresholds>
-				<colourThresholds stat="MVP30"> 0, 2, 3, 4</colourThresholds>
-				<colourThresholds stat="MemoryFreeMB">400,300,200,100</colourThresholds>
-				<colourThresholds stat="MemoryFreeMB(min)">400,300,200,50</colourThresholds>
-				<colourThresholds stat="DynamicResolutionPercentage">90,80,70,60</colourThresholds>
-				<colourThresholds stat="Frametime"> 33.4, 35, 38, 45  </colourThresholds>
-				<colourThresholds stat="GameThreadtime"> 33.4, 35, 38, 45  </colourThresholds>
-				<colourThresholds stat="RHIThreadTime"> 33.4, 35, 38, 45  </colourThresholds>
-				<colourThresholds stat="RenderThreadtime"> 33.4, 35, 38, 45  </colourThresholds>
-				<colourThresholds stat="GPUtime"> 33.4, 35, 38, 45  </colourThresholds>
-				<colourThresholds stat="Basic/TicksQueued"> 150, 200, 250, 300 </colourThresholds>
-			</summary>
-
-			<summary type="hitches">
-				<stats>Frametime,GameThreadtime,RenderThreadtime,RHIThreadTime,GPUtime</stats>
-				<colourThresholds> 20, 16.6, 16.2</colourThresholds>
-				<hitchThresholds>60, 150, 250, 500, 750, 1000, 2000</hitchThresholds>
-			</summary>
-
-
-			<summaryRef name="mapOverlay"/>
-			<!--
-			<summaryRef name="syncLoads"/>
-			<summaryRef name="textureStreaming"/>
-			<summaryRef name="psos"/>
-			-->
-
-			<graph title="Stat Unit 30hz" budget="33.33" />
-			<graph title="Stat Unit Raw" budget="33.33" />
-			<graph title="Gamethread Breakdown" budget="33.33"/>
-			<graph title="Gamethread Breakdown Raw" budget="33.33"/>
-			<graph title="Gamethread Waits" budget="4"/>
-			<graph title="Gamethread Physics" budget="4"/>
-			<graph title="Gamethread Breakdown Giant" budget="33.33" external="1"/>
-			<graph title="Ticks Breakdown" budget="200"/>
-			<graph title="Actor Counts" budget="0"/>
-			<graph title="Physical Memory" budget="300"/>
-			<graph title="Cmds DayaLayers" budget="300"/>
-			<graph title="Physical Memory Over Budget" budget="300"/>
-			<graph title="Dynamic Resolution" budget="80"/>
-			<graph title="GPU Breakdown" budget="33.33" />
-			<graph title="GPU Breakdown Raw" budget="16.66" />
-			<graph title="Niagara Breakdown" budget="33.33"/>
-			<graph title="RHI Drawcalls" budget="2000" />
-			<graph title="Renderthread Breakdown" budget="33.33"/>
-			<graph title="Renderthread Breakdown Raw" budget="33.33"/>
-			<graph title="FScene Breakdown" budget="33.33"/>
-			<graph title="FScene Breakdown Raw" budget="33.33"/>
-			<graph title="Render Targets Peak(MB)" budget="0" />
-			<graph title="Textures Peak(MB)" budget="0" />
-			<graph title="Virtual Texture" budget="0" />
-			<graph title="Light Counts" budget="0" />
-			<graph title="Animation Budget" budget="3.0" />
-			<graph title="Texture Streaming" budget="0" />
-			<graph title="Log Counts" budget="0" />
-			<graph title="CSV Profiler" budget="0" />
-		</reporttype>
-
-
-		<!-- Memory report -->
-		<reporttype name="Memory" title="Physical Memory Report" vsync="1">
-			<autodetection requiredstats="memoryfreemb,physicalusedmb"/>
-			<summary type="fpschart" fps="60" hitchThreshold="60" useEngineHitchMetric="1">
-				<stats>MemoryFreeMB(min),PhysicalUsedMB(max),MemoryFreeMB,PhysicalUsedMB,PhysicalUsedMB(min)</stats>
-				<colourThresholds stat="MemoryFreeMB(min)">400,300,200,50</colourThresholds>
-			</summary>
-			<graph title="Physical Memory" budget="300"/>
-		</reporttype>
-
-
-		<!-- LLM Report -->
-		<reporttype name="LLM" title="LLM Generic Report" ignoreList="" stripEvents="0">
-			<autodetection requiredstats="llm/total">
-				<requiredmetadata key="llm" allowedValues="1"/>
-			</autodetection>
-
-			<summary type="peak" hideStatPrefix="llm/">
-				<summarySection title="Main">
-					<statFilter>LLM/Total,LLM/TrackedTotal,LLM/Untracked,LLM/Textures,LLM/Audio,LLM/Meshes,LLM/RenderTargets,LLM/UObject,LLM/Animation,LLM/Shaders,LLM/Filesystem,LLM/Physics,LLM/ProgramSize</statFilter>
-				</summarySection>
-
-				<summarySection title="Audio">
-					<statFilter>LLM/Audio/*</statFilter>
-				</summarySection>
-
-				<summarySection title="Mesh">
-					<statFilter>LLM/Meshes/*</statFilter>
-				</summarySection>
-
-				<summarySection title="Physics">
-					<statFilter>LLM/Physics/*</statFilter>
-				</summarySection>
-
-				<summarySection title="Other">
-					<statFilter>LLM/*</statFilter>
-				</summarySection>
-			</summary>
-
-			<!-- top-level LLM summary graphs -->
-			<graph title="Total Memory Stats"			inSummary="1"/>
-
-			<!-- Memory breakdowns -->
-			<graph title="Tracked Memory Breakdown"		inSummary="1"/>
-			<graph title="Content Memory Breakdown"/>
-			<graph title="Audio Memory Breakdown"		inSummary="1"/>
-			<graph title="Mesh Memory Breakdown"		inSummary="1"/>
-			<graph title="Physics Memory Breakdown"		inSummary="1"/>
-
-			<!-- Various collections -->
-			<graph title="Rendering Memory Stats"		/>
-			<graph title="Misc Memory Stats"			/>
-			<graph title="Development Memory Stats"		/>
-			<graph title="Initialization Memory Stats"  />
-			<graph title="Loading Memory Stats"			/>
-
-			<!-- Individual stats -->
-			<graph title="Tracked Total Memory"         />
-			<graph title="UObject Memory"				/>
-			<graph title="Animation Memory"             />
-			<graph title="EngineMisc Memory"            />
-			<graph title="Materials Memory"             />
-			<graph title="RenderTargets Memory"         />
-			<graph title="Shaders Memory"               />
-			<graph title="Textures Memory"              />
-			<graph title="Untagged Memory"				inSummary="1"  budget="12"/>
-			<graph title="Untracked Memory"				inSummary="1"  budget="300"/>
-		</reporttype>
-
-		<!-- FX report -->
-		<reporttype name="FXReport" title="Playtest FX Perf Report" vsync="1">
-			<autodetection requiredstats="frametime">
-				<requiredmetadata key="targetframerate" allowedValues="60" stopIfKeyFound="1" ignoreIfKeyNotFound="1"/>
-			</autodetection>
-
-			<metadataToShow>testname,buildversion,deviceprofile,csvfilename,rhiname,desktophwinfo</metadataToShow>
-
-			<summary type="fpschart" fps="60" hitchThreshold="60" useEngineHitchMetric="1">
-				<stats>Frametime,GameThreadtime,RenderThreadtime,GPUtime,exclusive/gamethread/effects</stats>
-				<colourThresholds stat="Hitches/Min"> 2, 3, 5, 6</colourThresholds>
-				<colourThresholds stat="HitchTimePercent"> 0.1, 0.5, 1, 2</colourThresholds>
-				<colourThresholds stat="MVP60"> 2, 3, 4, 5</colourThresholds>
-				<colourThresholds stat="MemoryFreeMB(min)">400,300,200,50</colourThresholds>
-				<colourThresholds stat="Frametime"> 16.7, 17.5, 19, 22.5 </colourThresholds>
-				<colourThresholds stat="GameThreadtime"> 13, 15, 17, 22.5 </colourThresholds>
-				<colourThresholds stat="RenderThreadtime"> 13, 15, 17, 22.5 </colourThresholds>
-				<colourThresholds stat="GPUtime"> 16.7, 17.5, 19, 22.5 </colourThresholds>
-				<colourThresholds stat="exclusive/gamethread/effects"> 1.5, 2, 4, 6 </colourThresholds>			
-				<capture name="SkyDiving" startEvent="BusFlying" endEvent="SafeZones" includeFirstFrame="false" includeLastFrame="false" />
-				<capture name="OnTheGround" startEvent="SafeZones" endEvent="EndGame" includeFirstFrame="false" includeLastFrame="false" />
-				<capture name="MatchWithoutLobby" startEvent="BusFlying" endEvent="EndGame" includeFirstFrame="false" includeLastFrame="false" />
-			</summary>
-		
-			<summaryRef name="mapOverlay"/>
-			
-			<graph title="Gamethread Budget: Effects" budget="0.75"/>
-			<graph title="FX Breakdown Count" budget="50" />
-			<graph title="FX Breakdown Total" budget="2000.0" />
-			<graph title="FX Breakdown GTOnly" budget="1000.0" />			
-			<graph title="FX Breakdown RT" budget="1000.0" />	
-			<graph title="FX Breakdown GPU" budget="2000.0" />	
-			<graph title="FX Breakdown Activation" budget="300.0" />
-			<graph title="FX Breakdown Waits" budget="50.0" />
-			<graph title="FX Breakdown Per Inst Avg GT" budget="100.0" />
-			<graph title="FX Breakdown Per Inst Avg RT" budget="100.0" />
-			<graph title="FX Breakdown Per Inst Avg GPU" budget="100.0" />
-			<graph title="Niagara Culling" budget="0.0" />
-			<graph title="Niagara Culling Breakdown" budget="0.0" />
-			<graph title="Niagara Cull Proxy Breakdown" budget="0.0" />
-			<graph title="FX Budget Breakdown" budget="1.0" />
-			
-		</reporttype>
-
 	</reporttypes>
 </root>
\ No newline at end of file
