Upload perf gauntlet data to PRS

diff --git a/Engine/Plugins/VirtualProduction/ICVFXTesting/Build/Scripts/Automation/ICVFX.TestNode.cs b/Engine/Plugins/VirtualProduction/ICVFXTesting/Build/Scripts/Automation/ICVFX.TestNode.cs
--- a/Engine/Plugins/VirtualProduction/ICVFXTesting/Build/Scripts/Automation/ICVFX.TestNode.cs
+++ b/Engine/Plugins/VirtualProduction/ICVFXTesting/Build/Scripts/Automation/ICVFX.TestNode.cs
@@ -143,6 +143,27 @@ namespace ICVFXTest
 		///
 		[AutoParamWithNames("ICVFXTest.PerfCacheFolder")]
 		public string PerfCacheFolder;
+
+		/// <summary>
+		///  Whether to generate the local test reports.
+		/// </summary>
+		///
+		[AutoParamWithNames(false, "ICVFXTest.NoLocalReports")]
+		public bool NoLocalReports;
+
+		/// <summary>
+		///  Whether to skip uploading gauntlet perf data to the PRS.
+		/// </summary>
+		///
+		[AutoParamWithNames(false, "ICVFXTest.SkipPerfReportServer")]
+		public bool SkipPerfReportServer;
+
+		/// <summary>
+		///  Override the server import directory for PRS.
+		/// </summary>
+		///
+		[AutoParamWithNames("", "ICVFXTest.PerfReportServerImportDir")]
+		public string PerfReportServerImportDir;
 	}
 
 	public abstract class ICVFXTestNode : UnrealTestNode<ICVFXTestConfig>
@@ -292,4 +313,4 @@ namespace ICVFXTest
 			return base.CreateReport(Result, Context, Build, RoleResults, ArtifactPath);
 		}
 	}
-}
+}
\ No newline at end of file
diff --git a/Engine/Plugins/VirtualProduction/ICVFXTesting/Build/Scripts/Automation/Tests/ICVFXTest.PerformanceReport.cs b/Engine/Plugins/VirtualProduction/ICVFXTesting/Build/Scripts/Automation/Tests/ICVFXTest.PerformanceReport.cs
--- a/Engine/Plugins/VirtualProduction/ICVFXTesting/Build/Scripts/Automation/Tests/ICVFXTest.PerformanceReport.cs
+++ b/Engine/Plugins/VirtualProduction/ICVFXTesting/Build/Scripts/Automation/Tests/ICVFXTest.PerformanceReport.cs
@@ -89,20 +89,46 @@ namespace ICVFXTest
 	/// </summary>
 	public class PerformanceReport : AutoTest
 	{
+		private DirectoryInfo TempPerfCSVDir => new DirectoryInfo(Path.Combine(Unreal.RootDirectory.FullName, "GauntletTemp", "PerfReportCSVs"));
 		private static ILogger Logger => Log.Logger;
-		
+		private string OriginalBuildName;
+		private string OverrideDisplayConfigPath;
+		private string OverrideDisplayClusterNode;
+
 		public PerformanceReport(UnrealTestContext InContext)
 			: base(InContext)
 		{
 			FindSwitchboardConfigs();
+
+			// Save off the BuildName to prevent a mismatch in CreateReport
+			OriginalBuildName = Context.TestParams.ParseValue("ICVFXTest.BuildName", null);
+			Logger.LogInformation("Setting OriginalBuildName to {OriginalBuildName}", OriginalBuildName);
 		}
 
 		public override ICVFXTestConfig GetConfiguration()
 		{
 			ICVFXTestConfig Config = base.GetConfiguration();
 			UnrealTestRole ClientRole = Config.RequireRole(UnrealTargetRole.Client);
+
+			// Add CSV metadata
+			List<string> CsvMetadata = new List<string>
+			{
+				"testname=" + Config.TestName,
+				"gauntletTestType=AutoTest",
+				"gauntletSubTest=Performance",
+				"testBuildIsPreflight=" + (ReportGenUtils.IsTestingPreflightBuild(OriginalBuildName) ? "1" : "0"),
+				"testBuildVersion=" + OriginalBuildName
+			};
+
+			if (!string.IsNullOrEmpty(Context.BuildInfo.Branch) && Context.BuildInfo.Changelist != 0)
+			{
+				CsvMetadata.Add("branch=" + Context.BuildInfo.Branch);
+				CsvMetadata.Add("changelist=" + Context.BuildInfo.Changelist);
+			}
+
+			// Set CL parameters
 			ClientRole.CommandLineParams.Add("csvGpuStats");
-			ClientRole.CommandLineParams.Add("csvMetadata", $"\"testname={Config.TestName}\"");
+			ClientRole.CommandLineParams.Add("csvMetadata", "\"" + String.Join(",", CsvMetadata) + "\"");
 			ClientRole.CommandLineParams.AddOrAppendParamValue("execcmds", "t.FPSChart.DoCSVProfile 1");
 			ClientRole.CommandLineParams.AddOrAppendParamValue("execcmds", "t.FPSChart.OpenFolderOnDump 0");
 			ClientRole.CommandLineParams.Add("ICVFXTest.FPSChart");
@@ -209,6 +235,48 @@ namespace ICVFXTest
 			return OverrideDisplayClusterNode;
 		}
 
+		// Find the newest csv file and get its directory
+		private String GetNewestCSVDir(UnrealTargetPlatform Platform, string ArtifactPath, string TempDir) 
+		{
+			// All CSV paths to look through
+			var CSVsPaths = new[]
+			{
+				Path.Combine(ArtifactPath, "EditorGame", "Profiling", "FPSChartStats"),
+				Path.Combine(ArtifactPath, "EditorGame", "Settings", $"{Context.Options.Project}", "Saved", "Profiling", "FPSChartStats"),
+				Path.Combine(TempDir, "DeviceCache", Platform.ToString(), TestInstance.ClientApps[0].Device.ToString(), "UserDir")
+			};
+
+			// Check all paths for potential CSVs
+			var DiscoveredCSVs = new List<string>();
+			foreach (var CSVsPath in CSVsPaths)
+			{
+				if (Directory.Exists(CSVsPath))
+				{
+					DiscoveredCSVs.AddRange(
+						from CsvFile in Directory.GetFiles(CSVsPath, "*.csv", SearchOption.AllDirectories)
+						where CsvFile.Contains("csvprofile", StringComparison.InvariantCultureIgnoreCase)
+						select CsvFile);
+				}
+			}
+
+			if (DiscoveredCSVs.Count == 0)
+			{
+				Logger.LogError($"Test completed successfully but no CSV profiling results were found. Searched paths were:\r\n  {string.Join("\r\n  ", CSVsPaths.Select(s => $"\"{s}\""))}");
+				return null;
+			}
+
+			// Find the newest csv file and get its directory
+			// (PerfReportTool will only output cached data in -csvdir mode)
+			var NewestFile =
+				(from CsvFile in DiscoveredCSVs
+				 let Timestamp = File.GetCreationTimeUtc(CsvFile)
+				 orderby Timestamp descending
+				 select CsvFile).First();
+			var NewestDir = Path.GetDirectoryName(NewestFile);
+
+			return NewestDir;
+		}
+
 		/// <summary>
 		/// Produces a detailed csv report using PerfReportTool.
 		/// Also, stores perf data in the perf cache, and generates a historic report using the data the cache contains.
@@ -230,42 +298,13 @@ namespace ICVFXTest
 			}
 			var ReportConfigDir = Path.Combine(Unreal.RootDirectory.FullName, "Engine", "Plugins", "VirtualProduction", "ICVFXTesting", "Build", "Scripts", "PerfReport");
 			var ReportPath = Path.Combine(ArtifactPath, "Reports", "Performance");
+			var NewestDir = GetNewestCSVDir(Platform, ArtifactPath, TempDir);
 
-			var CsvsPaths = new[]
-				{
-					Path.Combine(ArtifactPath, "EditorGame", "Profiling", "FPSChartStats"),
-					Path.Combine(ArtifactPath, "EditorGame", "Settings", $"{Context.Options.Project}", "Saved", "Profiling", "FPSChartStats"),
-					Path.Combine(TempDir, "DeviceCache", Platform.ToString(), TestInstance.ClientApps[0].Device.ToString(), "UserDir")
-				};
-
-
-			var DiscoveredCsvs = new List<string>();
-				foreach (var CsvsPath in CsvsPaths)
-				{
-					if (Directory.Exists(CsvsPath))
-					{
-						DiscoveredCsvs.AddRange(
-							from CsvFile in Directory.GetFiles(CsvsPath, "*.csv", SearchOption.AllDirectories)
-							where CsvFile.Contains("csvprofile", StringComparison.InvariantCultureIgnoreCase)
-							select CsvFile);
-					}
-				}
-
-			if (DiscoveredCsvs.Count == 0)
-			{
-				Logger.LogError($"Test completed successfully but no csv profiling results were found. Searched paths were:\r\n  {string.Join("\r\n  ", CsvsPaths.Select(s => $"\"{s}\""))}");
+			if(string.IsNullOrEmpty(NewestDir)) { 
+				Logger.LogWarning("Failed to execute perf report tool");
 				return;
 			}
 
-			// Find the newest csv file and get its directory
-			// (PerfReportTool will only output cached data in -csvdir mode)
-			var NewestFile =
-				(from CsvFile in DiscoveredCsvs
-				 let Timestamp = File.GetCreationTimeUtc(CsvFile)
-				 orderby Timestamp descending
-				 select CsvFile).First();
-			var NewestDir = Path.GetDirectoryName(NewestFile);
-
 			Logger.LogInformation($"Using perf report cache directory \"{ReportCacheDir}\".");
 			Logger.LogInformation($"Using perf report output directory \"{ReportPath}\".");
 			Logger.LogInformation($"Using csv results directory \"{NewestDir}\". Generating historic perf report data...");
@@ -396,12 +435,73 @@ namespace ICVFXTest
 			});
 		}
 
+		public override bool StartTest(int Pass, int InNumPasses)
+		{
+			// Delete temporary folder for CSVs when rerunning
+			if (Pass == 0 && TempPerfCSVDir.Exists)
+			{
+				TempPerfCSVDir.Delete(recursive: true);
+			}
+			return base.StartTest(Pass, InNumPasses);
+		}
+
 		public override ITestReport CreateReport(TestResult Result, UnrealTestContext Context, UnrealBuildSource Build, IEnumerable<UnrealRoleResult> Artifacts, string ArtifactPath)
 		{
 			if (Result == TestResult.Passed)
 			{
-				Logger.LogInformation($"Generating performance reports using PerfReportTool.");
-				GeneratePerfReport(Context.GetRoleContext(UnrealTargetRole.Client).Platform, ArtifactPath, Context.Options.TempDir);
+				ICVFXTestConfig Config = base.GetConfiguration();
+
+				// Temporarily copy artifacts
+				CopyPerfFilesToTempDir(Context.GetRoleContext(UnrealTargetRole.Client).Platform, ArtifactPath, Context.Options.TempDir);
+
+				// Preserve artifacts till ready to be made into a report on the final iteration
+				if (GetCurrentPass() < (GetNumPasses() - 1))
+				{
+					Logger.LogInformation($"Skipping CSV report generator until final pass. On pass {GetCurrentPass() + 1} of {GetNumPasses()}.");
+					return base.CreateReport(Result, Context, Build, Artifacts, ArtifactPath);
+				}
+
+				// Create a local report
+				if (!Config.NoLocalReports)
+				{
+					Logger.LogInformation($"Generating performance reports using PerfReportTool.");
+					GeneratePerfReport(Context.GetRoleContext(UnrealTargetRole.Client).Platform, ArtifactPath, Context.Options.TempDir);
+				}
+
+				// Build import entries for PRS
+				if (!Config.SkipPerfReportServer)
+				{
+					Logger.LogInformation("Creating perf server importer with build name {BuildName}", OriginalBuildName);
+
+					string DataSourceName = "Automation.Saloon.ICVFXTesting";
+					string ImportDirOverride = Config.PerfReportServerImportDir;
+					Dictionary<string, dynamic> CommonDataSourceFields = new Dictionary<string, dynamic>
+					{
+						{ "HordeJobUrl", Globals.Params.ParseValue("JobDetails", null) }
+					};
+
+					ICsvImporter Importer = ReportGenUtils.CreatePerfReportServerImporter(DataSourceName, OriginalBuildName, CommandUtils.IsBuildMachine, ImportDirOverride, CommonDataSourceFields);
+					if (Importer != null)
+					{
+						// Recursively grab all the csv files we copied to the temp dir and convert them to binary.
+						List<FileInfo> AllBinaryCsvFiles = ReportGenUtils.CollectAndConvertCsvFilesToBinary(TempPerfCSVDir.FullName);
+						if (AllBinaryCsvFiles.Count == 0)
+						{
+							throw new AutomationException($"No CSV files found in {TempPerfCSVDir}");
+						}
+
+						// The corresponding log for each csv sits in the same subdirectory as the csv file itself.
+						IEnumerable<CsvImportEntry> ImportEntries = AllBinaryCsvFiles
+							.Select(CsvFile => new CsvImportEntry(CsvFile.FullName, Path.Combine(CsvFile.Directory.FullName, "ClientOutput.log")));
+
+						// Create the import batch
+						Logger.LogInformation("Importing entries to {DataSourceName}", DataSourceName);
+						Importer.Import(ImportEntries);
+					}
+
+					// Cleanup the temp dir
+					TempPerfCSVDir.Delete(recursive: true);
+				}
 			}
 			else
 			{
@@ -411,8 +511,63 @@ namespace ICVFXTest
 			return base.CreateReport(Result, Context, Build, Artifacts, ArtifactPath);
 		}
 
-		private string OverrideDisplayConfigPath;
-		private string OverrideDisplayClusterNode;
+		// Copy CSV files to temp directory because of duplicate import batches
+		private void CopyPerfFilesToTempDir(UnrealTargetPlatform Platform, string ArtifactPath, string TempDir)
+		{
+			if (!TempPerfCSVDir.Exists)
+			{
+				Logger.LogInformation("Creating temp perf CSV dir: {TempPerfCSVDir}", TempPerfCSVDir);
+				TempPerfCSVDir.Create();
+			}
+
+			string TestInstancePath = Path.Combine(TempDir, "DeviceCache", Platform.ToString(), TestInstance.ClientApps[0].Device.ToString());
+			string ClientLogPath = Path.Combine(TestInstancePath, "UserDir", "Saved", "Logs", "Saloon.log");
+			string FPSChartsPath = GetNewestCSVDir(Platform, ArtifactPath, TempDir);
+			
+			if (string.IsNullOrEmpty(FPSChartsPath))
+			{
+				Logger.LogWarning("Failed to find FPSCharts folder in {ArtifactPath}", ArtifactPath);
+				return;
+			}
+
+			FPSChartsPath = Gauntlet.Utils.SystemHelpers.GetFullyQualifiedPath(FPSChartsPath);
+			
+			Logger.LogInformation($"Using client log path \"{ClientLogPath}\".");
+			Logger.LogInformation($"Using FPS charts path \"{FPSChartsPath}\".");
+
+			// Grab all the csv files that have valid metadata
+			List<FileInfo> CsvFiles = ReportGenUtils.CollectValidCsvFiles(FPSChartsPath);
+			if (CsvFiles.Count > 0)
+			{
+				// We only want to copy the latest file as the other will have already been copied
+				CsvFiles.OrderBy(Info => Info.LastWriteTimeUtc);
+				FileInfo LatestCsvFile = CsvFiles.Last();
+
+				// Create a subdir for each pass as we want to store the csv and log together in the same dir to make it easier to find them
+				string PassDir = Path.Combine(TempPerfCSVDir.FullName, $"PerfCsv_Pass_{GetCurrentPass()}");
+				Directory.CreateDirectory(PassDir);
+
+				FileInfo LogFileInfo = new FileInfo(ClientLogPath);
+				if (LogFileInfo.Exists)
+				{
+					string LogDestPath = Path.Combine(PassDir, LogFileInfo.Name);
+					Logger.LogInformation("Copying Log {ClientLogPath} To {LogDest}", ClientLogPath, LogDestPath);
+					LogFileInfo.CopyTo(LogDestPath);
+				}
+				else
+				{
+					Logger.LogWarning("No log file was found at {ClientLogPath}", ClientLogPath);
+				}
+
+				string CsvDestPath = Path.Combine(PassDir, LatestCsvFile.Name);
+				Logger.LogInformation("Copying CSV {CsvPath} To {CsvDestPath}", LatestCsvFile.FullName, CsvDestPath);
+				LatestCsvFile.CopyTo(CsvDestPath);
+			}
+			else
+			{
+				Logger.LogWarning("No valid CSV files found in {FPSChartsPath}", FPSChartsPath);
+			}
+		}
 	}
   
 	//
diff --git a/Engine/Plugins/VirtualProduction/ICVFXTesting/Scripts/GauntletSettings.xml b/Engine/Plugins/VirtualProduction/ICVFXTesting/Scripts/GauntletSettings.xml
--- a/Engine/Plugins/VirtualProduction/ICVFXTesting/Scripts/GauntletSettings.xml
+++ b/Engine/Plugins/VirtualProduction/ICVFXTesting/Scripts/GauntletSettings.xml
@@ -46,6 +46,11 @@
 	<Option Name="WithVirtualShadowMaps" Restrict="true|false" DefaultValue="true" Description="Whether to use virtual shadow maps"/>
 	<Option Name="WithTrace" Restrict="true|false" DefaultValue="false" Description="Whether to run with unreal insights."/>
 	
+	<!-- PRS Specific Options -->
+	<Option Name="NoLocalReports" Restrict="true|false" DefaultValue="false" Description="Whether to generate the local test reports."/>
+	<Option Name="SkipPerfReportServer" Restrict="true|false" DefaultValue="false" Description="Whether to skip uploading gauntlet perf data to the PRS."/>
+	<Option Name="PerfReportServerImportDir" DefaultValue="" Description="Override the server import directory for PRS."/>
+	
 	
 	<Property Name="Platform" Value="Win64"/>
 
@@ -88,7 +93,11 @@
 	<Property Name="ExtraGauntletEditorTestArgs" Value="$(ExtraGauntletEditorTestArgs) -ICVFXTest.EnableTrace" If="$(WithTrace)"/>
 	<Property Name="ExtraGauntletEditorTestArgs" Value="$(ExtraGauntletEditorTestArgs) -ICVFXTest.TraceRootFolder=$(TestOutputDir)/$(BuildNamePath)" If="$(WithTrace)"/>
 	<Property Name="ExtraGauntletEditorTestArgs" Value="$(ExtraGauntletEditorTestArgs) -ICVFXTest.PerfCacheFolder=$(ReportCacheDirRoot)"/>
-	
+	<Property Name="ExtraGauntletEditorTestArgs" Value="$(ExtraGauntletEditorTestArgs) -ICVFXTest.NoLocalReports" If="$(NoLocalReports)"/>
+	<Property Name="ExtraGauntletEditorTestArgs" Value="$(ExtraGauntletEditorTestArgs) -ICVFXTest.SkipPerfReportServer" If="$(SkipPerfReportServer)"/>
+	<Property Name="ExtraGauntletEditorTestArgs" Value="$(ExtraGauntletEditorTestArgs) -ICVFXTest.PerfReportServerImportDir=$(PerfReportServerImportDir)"/>
+	<Property Name="ExtraGauntletEditorTestArgs" Value="$(ExtraGauntletEditorTestArgs) -ICVFXTest.BuildName=$(EscapedBranch)-CL-$(Change)" />
+
 	<!-- Common Gauntlet settings -->
 	<Include Script="../../../../../Engine/Build/Graph/Tasks/Inc/GauntletSettings.xml" />
 
