Fixed the PRS import directory being empty when running saloon perf tests on horde
Changed metadata to no longer be from development configs when running on horde which filters out PRS

diff --git a/Engine/Plugins/VirtualProduction/ICVFXTesting/Build/Scripts/Automation/ICVFX.TestNode.cs b/Engine/Plugins/VirtualProduction/ICVFXTesting/Build/Scripts/Automation/ICVFX.TestNode.cs
--- a/Engine/Plugins/VirtualProduction/ICVFXTesting/Build/Scripts/Automation/ICVFX.TestNode.cs
+++ b/Engine/Plugins/VirtualProduction/ICVFXTesting/Build/Scripts/Automation/ICVFX.TestNode.cs
@@ -164,6 +164,13 @@ namespace ICVFXTest
 		///
 		[AutoParamWithNames("", "ICVFXTest.PerfReportServerImportDir")]
 		public string PerfReportServerImportDir;
+
+		/// <summary>
+		/// Whether the server is running in development or not
+		/// </summary>
+		///
+		[AutoParamWithNames(false, "ICFVXTest.IsDevelopment")]
+		public bool IsDevelopment;
 	}
 
 	public abstract class ICVFXTestNode : UnrealTestNode<ICVFXTestConfig>
diff --git a/Engine/Plugins/VirtualProduction/ICVFXTesting/Build/Scripts/Automation/Tests/ICVFXTest.PerformanceReport.cs b/Engine/Plugins/VirtualProduction/ICVFXTesting/Build/Scripts/Automation/Tests/ICVFXTest.PerformanceReport.cs
--- a/Engine/Plugins/VirtualProduction/ICVFXTesting/Build/Scripts/Automation/Tests/ICVFXTest.PerformanceReport.cs
+++ b/Engine/Plugins/VirtualProduction/ICVFXTesting/Build/Scripts/Automation/Tests/ICVFXTest.PerformanceReport.cs
@@ -117,7 +117,8 @@ namespace ICVFXTest
 				"gauntletTestType=AutoTest",
 				"gauntletSubTest=Performance",
 				"testBuildIsPreflight=" + (ReportGenUtils.IsTestingPreflightBuild(OriginalBuildName) ? "1" : "0"),
-				"testBuildVersion=" + OriginalBuildName
+				"testBuildVersion=" + OriginalBuildName,
+				"config=" + (Config.IsDevelopment ? "development" : "test")
 			};
 
 			if (!string.IsNullOrEmpty(Context.BuildInfo.Branch) && Context.BuildInfo.Changelist != 0)
@@ -474,7 +475,14 @@ namespace ICVFXTest
 					Logger.LogInformation("Creating perf server importer with build name {BuildName}", OriginalBuildName);
 
 					string DataSourceName = "Automation.Saloon.ICVFXTesting";
-					string ImportDirOverride = Config.PerfReportServerImportDir;
+
+					string ImportDirOverride;
+					if(Config.PerfReportServerImportDir == "") {
+						ImportDirOverride = null;
+					} else {
+						ImportDirOverride = Config.PerfReportServerImportDir;
+					}
+					
 					Dictionary<string, dynamic> CommonDataSourceFields = new Dictionary<string, dynamic>
 					{
 						{ "HordeJobUrl", Globals.Params.ParseValue("JobDetails", null) }
diff --git a/Engine/Plugins/VirtualProduction/ICVFXTesting/Scripts/GauntletSettings.xml b/Engine/Plugins/VirtualProduction/ICVFXTesting/Scripts/GauntletSettings.xml
--- a/Engine/Plugins/VirtualProduction/ICVFXTesting/Scripts/GauntletSettings.xml
+++ b/Engine/Plugins/VirtualProduction/ICVFXTesting/Scripts/GauntletSettings.xml
@@ -50,7 +50,7 @@
 	<Option Name="NoLocalReports" Restrict="true|false" DefaultValue="false" Description="Whether to generate the local test reports."/>
 	<Option Name="SkipPerfReportServer" Restrict="true|false" DefaultValue="false" Description="Whether to skip uploading gauntlet perf data to the PRS."/>
 	<Option Name="PerfReportServerImportDir" DefaultValue="" Description="Override the server import directory for PRS."/>
-	
+	<Option Name="IsDevelopment" Restrict="true|false" DefaultValue="false" Description="Whether running in development or not."/>
 	
 	<Property Name="Platform" Value="Win64"/>
 
@@ -95,6 +95,7 @@
 	<Property Name="ExtraGauntletEditorTestArgs" Value="$(ExtraGauntletEditorTestArgs) -ICVFXTest.PerfCacheFolder=$(ReportCacheDirRoot)"/>
 	<Property Name="ExtraGauntletEditorTestArgs" Value="$(ExtraGauntletEditorTestArgs) -ICVFXTest.NoLocalReports" If="$(NoLocalReports)"/>
 	<Property Name="ExtraGauntletEditorTestArgs" Value="$(ExtraGauntletEditorTestArgs) -ICVFXTest.SkipPerfReportServer" If="$(SkipPerfReportServer)"/>
+	<Property Name="ExtraGauntletEditorTestArgs" Value="$(ExtraGauntletEditorTestArgs) -ICVFXTest.IsDevelopment" If="$(IsDevelopment)"/>
 	<Property Name="ExtraGauntletEditorTestArgs" Value="$(ExtraGauntletEditorTestArgs) -ICVFXTest.PerfReportServerImportDir=$(PerfReportServerImportDir)"/>
 	<Property Name="ExtraGauntletEditorTestArgs" Value="$(ExtraGauntletEditorTestArgs) -ICVFXTest.BuildName=$(EscapedBranch)-CL-$(Change)" />
 
