Do not set an agent override for compile
#rnx

diff --git a/Engine/Plugins/VirtualProduction/ICVFXTesting/Scripts/ICVFXTesting.xml b/Engine/Plugins/VirtualProduction/ICVFXTesting/Scripts/ICVFXTesting.xml
--- a/Engine/Plugins/VirtualProduction/ICVFXTesting/Scripts/ICVFXTesting.xml
+++ b/Engine/Plugins/VirtualProduction/ICVFXTesting/Scripts/ICVFXTesting.xml
@@ -30,9 +30,7 @@
 		<Property Name="DefaultTargetTestList" Value="" />
 	</Do>
 
-	<Property Name="AgentOverride" Value="EditorPerfWin64" />
 	<Property Name="TestAgentOverride" Value="EditorPerfWin64" />
-	<Property Name="ClientCompileAgentOverride" Value="EditorPerfWin64" />
 	<!-- Common Setting Overrides -->
 	
 	<Property Name="ICVFXTestingPath" Value="$(RootDir)/Engine/Plugins/VirtualProduction/ICVFXTesting/Scripts"/>
