Include preflight suffix within build name to seperate PFs and automated horde jobs in performance data

diff --git a/Engine/Plugins/VirtualProduction/ICVFXTesting/Scripts/GauntletSettings.xml b/Engine/Plugins/VirtualProduction/ICVFXTesting/Scripts/GauntletSettings.xml
--- a/Engine/Plugins/VirtualProduction/ICVFXTesting/Scripts/GauntletSettings.xml
+++ b/Engine/Plugins/VirtualProduction/ICVFXTesting/Scripts/GauntletSettings.xml
@@ -97,7 +97,7 @@
 	<Property Name="ExtraGauntletEditorTestArgs" Value="$(ExtraGauntletEditorTestArgs) -ICVFXTest.SkipPerfReportServer" If="$(SkipPerfReportServer)"/>
 	<Property Name="ExtraGauntletEditorTestArgs" Value="$(ExtraGauntletEditorTestArgs) -ICVFXTest.IsDevelopment" If="$(IsDevelopment)"/>
 	<Property Name="ExtraGauntletEditorTestArgs" Value="$(ExtraGauntletEditorTestArgs) -ICVFXTest.PerfReportServerImportDir=$(PerfReportServerImportDir)"/>
-	<Property Name="ExtraGauntletEditorTestArgs" Value="$(ExtraGauntletEditorTestArgs) -ICVFXTest.BuildName=$(EscapedBranch)-CL-$(Change)" />
+	<Property Name="ExtraGauntletEditorTestArgs" Value="$(ExtraGauntletEditorTestArgs) -ICVFXTest.BuildName=$(EscapedBranch)-CL-$(Change)$(PreflightSuffix)" />
 
 	<!-- Common Gauntlet settings -->
 	<Include Script="../../../../../Engine/Build/Graph/Tasks/Inc/GauntletSettings.xml" />
