Remove import to PRS

diff --git a/Engine/Plugins/VirtualProduction/ICVFXTesting/Build/Scripts/Automation/Tests/ICVFXTest.PerformanceReport.cs b/Engine/Plugins/VirtualProduction/ICVFXTesting/Build/Scripts/Automation/Tests/ICVFXTest.PerformanceReport.cs
--- a/Engine/Plugins/VirtualProduction/ICVFXTesting/Build/Scripts/Automation/Tests/ICVFXTest.PerformanceReport.cs
+++ b/Engine/Plugins/VirtualProduction/ICVFXTesting/Build/Scripts/Automation/Tests/ICVFXTest.PerformanceReport.cs
@@ -504,7 +504,7 @@ namespace ICVFXTest
 
 						// Create the import batch
 						Logger.LogInformation("Importing entries to {DataSourceName}", DataSourceName);
-						Importer.Import(ImportEntries);
+						//Importer.Import(ImportEntries);
 					}
 
 					// Cleanup the temp dir
