Update comments

diff --git a/Engine/Plugins/VirtualProduction/ICVFXTesting/Build/Scripts/Automation/Tests/ICVFXTest.PerformanceReport.cs b/Engine/Plugins/VirtualProduction/ICVFXTesting/Build/Scripts/Automation/Tests/ICVFXTest.PerformanceReport.cs
--- a/Engine/Plugins/VirtualProduction/ICVFXTesting/Build/Scripts/Automation/Tests/ICVFXTest.PerformanceReport.cs
+++ b/Engine/Plugins/VirtualProduction/ICVFXTesting/Build/Scripts/Automation/Tests/ICVFXTest.PerformanceReport.cs
@@ -279,8 +279,8 @@ namespace ICVFXTest
 		}
 
 		/// <summary>
-		/// Produces a detailed csv report using PerfReportTool.
-		/// Also, stores perf data in the perf cache, and generates a historic report using the data the cache contains.
+		/// Stores perf data in local perf cache which can be uploaded to PRS, 
+		/// and produces a detailed csv report using PerfReportTool using the data the cache contains.
 		/// </summary>
 		private void GeneratePerfReport(UnrealTargetPlatform Platform, string ArtifactPath, string TempDir)
 		{
@@ -415,20 +415,20 @@ namespace ICVFXTest
 				}
 			}
 
-			// All-time historic summary report
+			// Creates/Updates all-time historic summary report
 			HistoricReport_Alt("HistoricReport_AllTime", new[]
 			{
 				$"platform={PlatformNameFilter}"
 			});
 
-			// 14 days historic report
+			// Creates/Updates 14 days historic report
 			HistoricReport_Alt($"HistoricReport_14Days", new[]
 			{
 				$"platform={PlatformNameFilter}",
 				$"starttimestamp>={DateTimeOffset.Now.ToUnixTimeSeconds() - (14 * 60L * 60L * 24L)}"
 			});
 
-			// 14 days historic report
+			// Creates/Updates 14 days historic report
 			HistoricReport($"HistoricReport_14Days_Summary", new[]
 			{
 				$"platform={PlatformNameFilter}",
@@ -504,7 +504,7 @@ namespace ICVFXTest
 
 						// Create the import batch
 						Logger.LogInformation("Importing entries to {DataSourceName}", DataSourceName);
-						//Importer.Import(ImportEntries);
+						Importer.Import(ImportEntries);
 					}
 
 					// Cleanup the temp dir
