Repointed ICVFX.Automation dll to Engine\Binaries\DotNET\AutomationTool\AutomationScripts so it can be found by build

#rnx

diff --git a/Engine/Plugins/VirtualProduction/ICVFXTesting/Build/Scripts/ICVFX.Automation.csproj b/Engine/Plugins/VirtualProduction/ICVFXTesting/Build/Scripts/ICVFX.Automation.csproj
--- a/Engine/Plugins/VirtualProduction/ICVFXTesting/Build/Scripts/ICVFX.Automation.csproj
+++ b/Engine/Plugins/VirtualProduction/ICVFXTesting/Build/Scripts/ICVFX.Automation.csproj
@@ -18,7 +18,7 @@
     <RootNamespace>ICVFX.Automation</RootNamespace>
     <AssemblyName>ICVFX.Automation</AssemblyName>
     <WarningsNotAsErrors>612,618</WarningsNotAsErrors>
-    <OutputPath>..\..\Binaries\DotNET\AutomationScripts\</OutputPath>
+    <OutputPath>..\..\..\..\..\Binaries\DotNET\AutomationTool\AutomationScripts\</OutputPath>
     <TreatWarningsAsErrors>true</TreatWarningsAsErrors>
     <DebugType>pdbonly</DebugType>
   </PropertyGroup>
