Tweak automation script .csproj to match setup of other scripts

#rnx

diff --git a/Engine/Plugins/VirtualProduction/ICVFXTesting/Build/Scripts/ICVFX.Automation.csproj b/Engine/Plugins/VirtualProduction/ICVFXTesting/Build/Scripts/ICVFX.Automation.csproj
--- a/Engine/Plugins/VirtualProduction/ICVFXTesting/Build/Scripts/ICVFX.Automation.csproj
+++ b/Engine/Plugins/VirtualProduction/ICVFXTesting/Build/Scripts/ICVFX.Automation.csproj
@@ -1,11 +1,6 @@
 <Project Sdk="Microsoft.NET.Sdk">
-  <Import Project="ICVFX.Automation.csproj.props" Condition="Exists('ICVFX.Automation.csproj.props')"/>
-
-  <PropertyGroup>
-    <EngineDir>../../../../..</EngineDir>
-  </PropertyGroup>
-  
-  <Import Project="$(EngineDir)\Source\Programs\Shared\UnrealEngine.csproj.props" />
+  <Import Project="ICVFX.Automation.csproj.props" Condition="Exists('ICVFX.Automation.csproj.props')"/>  
+  <Import Project="..\..\..\..\..\Source\Programs\Shared\UnrealEngine.csproj.props" />
   
   <PropertyGroup>
     <TargetFramework>net8.0</TargetFramework>
