[PerfTesting] Engine part - fix for offscreen nodes
- Fixed  cmdline parameters required for offscreen launch
- Fixed passing of the offscreen, nanite and lumen switches to automation

#rnx

diff --git a/Engine/Plugins/VirtualProduction/ICVFXTesting/Build/Scripts/Automation/Tests/ICVFXTest.AutoTest.cs b/Engine/Plugins/VirtualProduction/ICVFXTesting/Build/Scripts/Automation/Tests/ICVFXTest.AutoTest.cs
--- a/Engine/Plugins/VirtualProduction/ICVFXTesting/Build/Scripts/Automation/Tests/ICVFXTest.AutoTest.cs
+++ b/Engine/Plugins/VirtualProduction/ICVFXTesting/Build/Scripts/Automation/Tests/ICVFXTest.AutoTest.cs
@@ -238,8 +238,6 @@ namespace ICVFXTest
 			}
 			else
 			{
-				ClientRole.CommandLineParams.Add("forceres");
-				ClientRole.CommandLineParams.Add("windowed");
 				ClientRole.CommandLineParams.Add($"WinX={Config.WndLocX}");
 				ClientRole.CommandLineParams.Add($"WinY={Config.WndLocY}");
 				ClientRole.CommandLineParams.Add($"ResX={Config.WndResX}");
@@ -250,6 +248,11 @@ namespace ICVFXTest
 				{
 					ClientRole.CommandLineParams.Add("RenderOffscreen");
 				}
+				else
+				{
+					ClientRole.CommandLineParams.Add("forceres");
+					ClientRole.CommandLineParams.Add("windowed");
+				}
 			}
 
 			Gauntlet.Log.Info(ClientRole.CommandLineParams.ToString());
diff --git a/Engine/Plugins/VirtualProduction/ICVFXTesting/Scripts/GauntletSettings.xml b/Engine/Plugins/VirtualProduction/ICVFXTesting/Scripts/GauntletSettings.xml
--- a/Engine/Plugins/VirtualProduction/ICVFXTesting/Scripts/GauntletSettings.xml
+++ b/Engine/Plugins/VirtualProduction/ICVFXTesting/Scripts/GauntletSettings.xml
@@ -104,8 +104,8 @@
 	<Property Name="ExtraGauntletEditorTestArgs" Value="$(ExtraGauntletEditorTestArgs) -ICVFXTest.MapOverride=$(MapOverride)" />
 	<Property Name="ExtraGauntletEditorTestArgs" Value="$(ExtraGauntletEditorTestArgs) -ICVFXTest.DisplayConfig=$(DisplayConfigPath)" />
 	<Property Name="ExtraGauntletEditorTestArgs" Value="$(ExtraGauntletEditorTestArgs) -ICVFXTest.ViewportScreenPercentage=$(ViewportScreenPercentage)" />
-	<Property Name="ExtraGauntletEditorTestArgs" Value="$(ExtraGauntletEditorTestArgs) -ICVFXTest.Fullscreen=$(Fullscreen)" />
-	<Property Name="ExtraGauntletEditorTestArgs" Value="$(ExtraGauntletEditorTestArgs) -ICVFXTest.Offscreen=$(Offscreen)" />
+	<Property Name="ExtraGauntletEditorTestArgs" Value="$(ExtraGauntletEditorTestArgs) -ICVFXTest.Fullscreen" If="$(Fullscreen)" />
+	<Property Name="ExtraGauntletEditorTestArgs" Value="$(ExtraGauntletEditorTestArgs) -ICVFXTest.Offscreen" If="$(Offscreen)" />
 	<Property Name="ExtraGauntletEditorTestArgs" Value="$(ExtraGauntletEditorTestArgs) -ICVFXTest.WndLocX=$(WndLocX)" />
 	<Property Name="ExtraGauntletEditorTestArgs" Value="$(ExtraGauntletEditorTestArgs) -ICVFXTest.WndLocY=$(WndLocY)" />
 	<Property Name="ExtraGauntletEditorTestArgs" Value="$(ExtraGauntletEditorTestArgs) -ICVFXTest.WndResX=$(WndResX)" />
@@ -113,8 +113,8 @@
 	<Property Name="ExtraGauntletEditorTestArgs" Value="$(ExtraGauntletEditorTestArgs) -ICVFXTest.DisplayClusterNodeName=$(DisplayClusterNodeName)" />
 	<Property Name="ExtraGauntletEditorTestArgs" Value="$(ExtraGauntletEditorTestArgs) -ICVFXTest.D3DDebug" If="$(WithD3DDebug)"/>
 	<Property Name="ExtraGauntletEditorTestArgs" Value="$(ExtraGauntletEditorTestArgs) -ICVFXTest.GPUCrashDebugging" If="$(WithGPUCrashDebugging)"/>
-	<Property Name="ExtraGauntletEditorTestArgs" Value="$(ExtraGauntletEditorTestArgs) -ICVFXTest.Nanite=$(WithNanite)"/>
-	<Property Name="ExtraGauntletEditorTestArgs" Value="$(ExtraGauntletEditorTestArgs) -ICVFXTest.Lumen=$(WithLumen)"/>
+	<Property Name="ExtraGauntletEditorTestArgs" Value="$(ExtraGauntletEditorTestArgs) -ICVFXTest.Nanite" If="$(WithNanite)"/>
+	<Property Name="ExtraGauntletEditorTestArgs" Value="$(ExtraGauntletEditorTestArgs) -ICVFXTest.Lumen=" If="$(WithLumen)"/>
 	<Property Name="ExtraGauntletEditorTestArgs" Value="$(ExtraGauntletEditorTestArgs) -ICVFXTest.MaxGPUCount=$(MaxGPUCount)"/>
 	<Property Name="ExtraGauntletEditorTestArgs" Value="$(ExtraGauntletEditorTestArgs) -ICVFXTest.SummaryReportType=$(SummaryReportType)" If="'$(SummaryReportType)' != ''" />
 	<Property Name="ExtraGauntletEditorTestArgs" Value="$(ExtraGauntletEditorTestArgs) -ICVFXTest.SummaryReportPath=$(SummaryReportPath)"/>
