[PerfTesting] Fixed typo in the test script

#rnx

diff --git a/Engine/Plugins/VirtualProduction/ICVFXTesting/Scripts/GauntletSettings.xml b/Engine/Plugins/VirtualProduction/ICVFXTesting/Scripts/GauntletSettings.xml
--- a/Engine/Plugins/VirtualProduction/ICVFXTesting/Scripts/GauntletSettings.xml
+++ b/Engine/Plugins/VirtualProduction/ICVFXTesting/Scripts/GauntletSettings.xml
@@ -114,7 +114,7 @@
 	<Property Name="ExtraGauntletEditorTestArgs" Value="$(ExtraGauntletEditorTestArgs) -ICVFXTest.D3DDebug" If="$(WithD3DDebug)"/>
 	<Property Name="ExtraGauntletEditorTestArgs" Value="$(ExtraGauntletEditorTestArgs) -ICVFXTest.GPUCrashDebugging" If="$(WithGPUCrashDebugging)"/>
 	<Property Name="ExtraGauntletEditorTestArgs" Value="$(ExtraGauntletEditorTestArgs) -ICVFXTest.Nanite" If="$(WithNanite)"/>
-	<Property Name="ExtraGauntletEditorTestArgs" Value="$(ExtraGauntletEditorTestArgs) -ICVFXTest.Lumen=" If="$(WithLumen)"/>
+	<Property Name="ExtraGauntletEditorTestArgs" Value="$(ExtraGauntletEditorTestArgs) -ICVFXTest.Lumen" If="$(WithLumen)"/>
 	<Property Name="ExtraGauntletEditorTestArgs" Value="$(ExtraGauntletEditorTestArgs) -ICVFXTest.MaxGPUCount=$(MaxGPUCount)"/>
 	<Property Name="ExtraGauntletEditorTestArgs" Value="$(ExtraGauntletEditorTestArgs) -ICVFXTest.SummaryReportType=$(SummaryReportType)" If="'$(SummaryReportType)' != ''" />
 	<Property Name="ExtraGauntletEditorTestArgs" Value="$(ExtraGauntletEditorTestArgs) -ICVFXTest.SummaryReportPath=$(SummaryReportPath)"/>
