RC Web Interface: fixed issue with building on mac os where the bin directory added to PATH in the Start script was not absolute causing the subsequent web app scripts to fail

diff --git a/Engine/Plugins/VirtualProduction/RemoteControlWebInterface/WebApp/Start.sh b/Engine/Plugins/VirtualProduction/RemoteControlWebInterface/WebApp/Start.sh
--- a/Engine/Plugins/VirtualProduction/RemoteControlWebInterface/WebApp/Start.sh
+++ b/Engine/Plugins/VirtualProduction/RemoteControlWebInterface/WebApp/Start.sh
@@ -10,6 +10,7 @@ fi
 NodeVersion=v16.17.0
 DownloadedNodeFolder=node-${NodeVersion}-${NodeOS}-x64
 NodeName=${DownloadedNodeFolder}.tar.gz
+NodeFolderBasePath=$(cd "$(dirname "$0")"; pwd -P)
 
 if [ ! -d "${DownloadedNodeFolder}" ]; then
   echo Downloading NodeJS ...
@@ -35,7 +36,7 @@ fi
 if [ -d "${DownloadedNodeFolder}" ]; then
   # Add downloaded nodejs version to be first in line
   echo Using downloaded NodeJS version
-  export PATH="${DownloadedNodeFolder}/bin:$PATH"
+  export PATH="${NodeFolderBasePath}/${DownloadedNodeFolder}/bin:$PATH"
 fi
 
 
