Fixed Web RC not building in mac on installed paths with spaces

diff --git a/Engine/Plugins/VirtualProduction/RemoteControlWebInterface/WebApp/Start.sh b/Engine/Plugins/VirtualProduction/RemoteControlWebInterface/WebApp/Start.sh
--- a/Engine/Plugins/VirtualProduction/RemoteControlWebInterface/WebApp/Start.sh
+++ b/Engine/Plugins/VirtualProduction/RemoteControlWebInterface/WebApp/Start.sh
@@ -69,5 +69,4 @@ fi
 
 
 FOLDER=$(dirname "$0")
-node ${FOLDER}/Scripts/start.js "$@"
-
+node "${FOLDER}/Scripts/start.js" "$@"
