AutomationScripts: Use default TargetFramework from shared property sheet

#was-backout 50920539
#rnx

#Synced-CL 50923123

diff --git a/Engine/Plugins/Performance/AutomatedPerfTesting/Build/Scripts/AutomatedPerfTesting.Automation.csproj b/Engine/Plugins/Performance/AutomatedPerfTesting/Build/Scripts/AutomatedPerfTesting.Automation.csproj
--- a/Engine/Plugins/Performance/AutomatedPerfTesting/Build/Scripts/AutomatedPerfTesting.Automation.csproj
+++ b/Engine/Plugins/Performance/AutomatedPerfTesting/Build/Scripts/AutomatedPerfTesting.Automation.csproj
@@ -2,17 +2,13 @@
   <Import Project="AutomatedPerfTesting.Automation.csproj.props" Condition="Exists('AutomatedPerfTesting.Automation.csproj.props')" />
   <Import Project="..\..\..\..\..\Source\Programs\Shared\UnrealEngine.csproj.props" />
   <PropertyGroup>
-    <TargetFramework>net8.0</TargetFramework>
     <OutputType>Library</OutputType>
-    <AppendTargetFrameworkToOutputPath>false</AppendTargetFrameworkToOutputPath>
     <GenerateTargetFrameworkAttribute>false</GenerateTargetFrameworkAttribute>
-    <Configurations>Debug;Release;Development</Configurations>
     <RootNamespace>AutomatedPerfTesting.Automation</RootNamespace>
     <AssemblyName>AutomatedPerfTesting.Automation</AssemblyName>
     <OutputPath>$(EngineDir)\Binaries\DotNET\AutomationTool\AutomationScripts\AutomatedPerfTesting</OutputPath>
     <TreatWarningsAsErrors>true</TreatWarningsAsErrors>
     <WarningsNotAsErrors>612,618</WarningsNotAsErrors>
-    <DebugType>pdbonly</DebugType>
   </PropertyGroup>
   <ItemGroup>
     <ProjectReference Include="$(EngineDir)\Source\Programs\AutomationTool\AutomationUtils\AutomationUtils.Automation.csproj" />
diff --git a/Engine/Plugins/VirtualProduction/ICVFXTesting/Build/Scripts/ICVFX.Automation.csproj b/Engine/Plugins/VirtualProduction/ICVFXTesting/Build/Scripts/ICVFX.Automation.csproj
--- a/Engine/Plugins/VirtualProduction/ICVFXTesting/Build/Scripts/ICVFX.Automation.csproj
+++ b/Engine/Plugins/VirtualProduction/ICVFXTesting/Build/Scripts/ICVFX.Automation.csproj
@@ -1,56 +1,31 @@
 <Project Sdk="Microsoft.NET.Sdk">
-  <Import Project="ICVFX.Automation.csproj.props" Condition="Exists('ICVFX.Automation.csproj.props')"/>  
+  <Import Project="ICVFX.Automation.csproj.props" Condition="Exists('ICVFX.Automation.csproj.props')" />
   <Import Project="..\..\..\..\..\Source\Programs\Shared\UnrealEngine.csproj.props" />
-  
   <PropertyGroup>
-    <TargetFramework>net8.0</TargetFramework>
-    <Configuration Condition=" '$(Configuration)' == '' ">Development</Configuration>
     <OutputType>Library</OutputType>
-    <AppendTargetFrameworkToOutputPath>false</AppendTargetFrameworkToOutputPath>
     <GenerateAssemblyInfo>false</GenerateAssemblyInfo>
     <GenerateTargetFrameworkAttribute>false</GenerateTargetFrameworkAttribute>
-    <Configurations>Debug;Release;Development</Configurations>
     <RootNamespace>ICVFX.Automation</RootNamespace>
     <AssemblyName>ICVFX.Automation</AssemblyName>
     <WarningsNotAsErrors>612,618</WarningsNotAsErrors>
     <OutputPath>..\..\..\..\..\Binaries\DotNET\AutomationTool\AutomationScripts\ICVFX</OutputPath>
     <TreatWarningsAsErrors>true</TreatWarningsAsErrors>
-    <DebugType>pdbonly</DebugType>
   </PropertyGroup>
-  <PropertyGroup Condition=" '$(Configuration)|$(Platform)' == 'Development|AnyCPU' ">
-    <DefineConstants>$(DefineConstants);TRACE</DefineConstants>
-    <Optimize>true</Optimize>
-  </PropertyGroup>
-  <PropertyGroup Condition=" '$(Configuration)|$(Platform)' == 'Release|AnyCPU' ">
-    <DefineConstants>$(DefineConstants);TRACE</DefineConstants>
-    <Optimize>true</Optimize>
-  </PropertyGroup>
-  <PropertyGroup Condition=" '$(Configuration)|$(Platform)' == 'Debug|AnyCPU' ">
-    <DefineConstants>$(DefineConstants);DEBUG;TRACE</DefineConstants>
-    <DebugSymbols>true</DebugSymbols>
-    <DebugType>full</DebugType>
-  </PropertyGroup>
-  <PropertyGroup>
-    <DefineConstants>TRACE</DefineConstants>
-  </PropertyGroup>
-
   <ItemGroup>
     <Compile Include="$(EngineDir)\Source\Programs\Shared\MetaData.cs">
       <Link>Properties\MetaData.cs</Link>
     </Compile>
   </ItemGroup>
-
   <ItemGroup>
     <PackageReference Include="Microsoft.CSharp" Version="4.7.0" />
     <PackageReference Include="System.Data.DataSetExtensions" Version="4.5.0" />
   </ItemGroup>
-
   <ItemGroup>
     <ProjectReference Include="$(EngineDir)\Source\Programs\AutomationTool\AutomationUtils\AutomationUtils.Automation.csproj" />
     <ProjectReference Include="$(EngineDir)\Source\Programs\AutomationTool\Gauntlet\Gauntlet.Automation.csproj" />
     <ProjectReference Include="$(EngineDir)\Source\Programs\AutomationTool\Scripts\AutomationScripts.Automation.csproj" />
-   	<ProjectReference Include="$(EngineDir)\Source\Programs\AutomationTool\Localization\Localization.Automation.csproj" />
-  	<ProjectReference Include="$(EngineDir)\Source\Programs\AutomationTool\XLocLocalization\XLocLocalization.Automation.csproj" />
+    <ProjectReference Include="$(EngineDir)\Source\Programs\AutomationTool\Localization\Localization.Automation.csproj" />
+    <ProjectReference Include="$(EngineDir)\Source\Programs\AutomationTool\XLocLocalization\XLocLocalization.Automation.csproj" />
     <ProjectReference Include="$(EngineDir)\Source\Programs\Shared\EpicGames.Core\EpicGames.Core.csproj" />
     <ProjectReference Include="$(EngineDir)\Source\Programs\Shared\EpicGames.Build\EpicGames.Build.csproj" />
     <ProjectReference Include="$(EngineDir)\Source\Programs\UnrealBuildTool\UnrealBuildTool.csproj" />
diff --git a/Engine/Source/Programs/AutomationTool/Android/Android.Automation.csproj b/Engine/Source/Programs/AutomationTool/Android/Android.Automation.csproj
--- a/Engine/Source/Programs/AutomationTool/Android/Android.Automation.csproj
+++ b/Engine/Source/Programs/AutomationTool/Android/Android.Automation.csproj
@@ -2,32 +2,14 @@
   <Import Project="..\..\Shared\UnrealEngine.csproj.props" />
   
   <PropertyGroup>
-    <TargetFramework>net8.0</TargetFramework>
-    <Configuration Condition=" '$(Configuration)' == '' ">Development</Configuration>
     <OutputType>Library</OutputType>
-    <AppendTargetFrameworkToOutputPath>false</AppendTargetFrameworkToOutputPath>
     <GenerateAssemblyInfo>false</GenerateAssemblyInfo>
     <GenerateTargetFrameworkAttribute>false</GenerateTargetFrameworkAttribute>
-    <Configurations>Debug;Release;Development</Configurations>
     <RootNamespace>Android.Automation</RootNamespace>
     <AssemblyName>Android.Automation</AssemblyName>
     <WarningsNotAsErrors>612,618</WarningsNotAsErrors>
     <OutputPath>..\..\..\..\Binaries\DotNET\AutomationTool\AutomationScripts\Platforms\Android</OutputPath>
     <TreatWarningsAsErrors>true</TreatWarningsAsErrors>
-    <DebugType>pdbonly</DebugType>
-  </PropertyGroup>
-  <PropertyGroup Condition=" '$(Configuration)|$(Platform)' == 'Development|AnyCPU' ">
-    <DefineConstants>$(DefineConstants);TRACE</DefineConstants>
-    <Optimize>true</Optimize>
-  </PropertyGroup>
-  <PropertyGroup Condition=" '$(Configuration)|$(Platform)' == 'Release|AnyCPU' ">
-    <DefineConstants>$(DefineConstants);TRACE</DefineConstants>
-    <Optimize>true</Optimize>
-  </PropertyGroup>
-  <PropertyGroup Condition=" '$(Configuration)|$(Platform)' == 'Debug|AnyCPU' ">
-    <DefineConstants>$(DefineConstants);DEBUG;TRACE</DefineConstants>
-    <DebugSymbols>true</DebugSymbols>
-    <DebugType>full</DebugType>
   </PropertyGroup>
   <ItemGroup>
     <Reference Include="Ionic.Zip.Reduced">
diff --git a/Engine/Source/Programs/AutomationTool/Apple/Apple.Automation.csproj b/Engine/Source/Programs/AutomationTool/Apple/Apple.Automation.csproj
--- a/Engine/Source/Programs/AutomationTool/Apple/Apple.Automation.csproj
+++ b/Engine/Source/Programs/AutomationTool/Apple/Apple.Automation.csproj
@@ -2,32 +2,14 @@
   <Import Project="..\..\Shared\UnrealEngine.csproj.props" />
   
   <PropertyGroup>
-    <TargetFramework>net8.0</TargetFramework>
-    <Configuration Condition=" '$(Configuration)' == '' ">Development</Configuration>
     <OutputType>Library</OutputType>
-    <AppendTargetFrameworkToOutputPath>false</AppendTargetFrameworkToOutputPath>
     <GenerateAssemblyInfo>false</GenerateAssemblyInfo>
     <GenerateTargetFrameworkAttribute>false</GenerateTargetFrameworkAttribute>
-    <Configurations>Debug;Release;Development</Configurations>
     <RootNamespace>Apple.Automation</RootNamespace>
     <AssemblyName>Apple.Automation</AssemblyName>
     <WarningsNotAsErrors>612,618</WarningsNotAsErrors>
     <OutputPath>..\..\..\..\Binaries\DotNET\AutomationTool\AutomationScripts\Platforms\Apple</OutputPath>
     <TreatWarningsAsErrors>true</TreatWarningsAsErrors>
-    <DebugType>pdbonly</DebugType>
-  </PropertyGroup>
-  <PropertyGroup Condition=" '$(Configuration)|$(Platform)' == 'Development|AnyCPU' ">
-    <DefineConstants>$(DefineConstants);TRACE</DefineConstants>
-    <Optimize>true</Optimize>
-  </PropertyGroup>
-  <PropertyGroup Condition=" '$(Configuration)|$(Platform)' == 'Release|AnyCPU' ">
-    <DefineConstants>$(DefineConstants);TRACE</DefineConstants>
-    <Optimize>true</Optimize>
-  </PropertyGroup>
-  <PropertyGroup Condition=" '$(Configuration)|$(Platform)' == 'Debug|AnyCPU' ">
-    <DefineConstants>$(DefineConstants);DEBUG;TRACE</DefineConstants>
-    <DebugSymbols>true</DebugSymbols>
-    <DebugType>full</DebugType>
   </PropertyGroup>
   <ItemGroup>
     <Compile Include="..\..\Shared\MetaData.cs">
diff --git a/Engine/Source/Programs/AutomationTool/AutomationUtils/AutomationUtils.Automation.csproj b/Engine/Source/Programs/AutomationTool/AutomationUtils/AutomationUtils.Automation.csproj
--- a/Engine/Source/Programs/AutomationTool/AutomationUtils/AutomationUtils.Automation.csproj
+++ b/Engine/Source/Programs/AutomationTool/AutomationUtils/AutomationUtils.Automation.csproj
@@ -2,18 +2,12 @@
   <Import Project="..\..\Shared\UnrealEngine.csproj.props" />
   
   <PropertyGroup>
-    <TargetFramework>net8.0</TargetFramework>
-    <Configuration Condition=" '$(Configuration)' == '' ">Development</Configuration>
     <OutputType>Library</OutputType>
     <RootNamespace>AutomationUtils</RootNamespace>
     <AssemblyName>AutomationUtils.Automation</AssemblyName>
     <GenerateAssemblyInfo>false</GenerateAssemblyInfo>
     <GenerateTargetFrameworkAttribute>false</GenerateTargetFrameworkAttribute>
-    <Configurations>Debug;Release;Development</Configurations>
     <TreatWarningsAsErrors>true</TreatWarningsAsErrors>
-    <AppendTargetFrameworkToOutputPath>false</AppendTargetFrameworkToOutputPath>
-    <DebugType>pdbonly</DebugType>
-    <DebugType Condition="'$(IsWindows)' == 'true' And '$(IsArm64)' == 'true'">portable</DebugType>
     <CopyLocalLockFileAssemblies>true</CopyLocalLockFileAssemblies>
   </PropertyGroup>
   <PropertyGroup Condition="'$(EpicGamesMsBuild)' == 'true'">
@@ -29,19 +23,6 @@
       <Private>False</Private>
     </Reference>
   </ItemGroup>
-  <PropertyGroup Condition=" '$(Configuration)|$(Platform)' == 'Debug|AnyCPU' ">
-    <DefineConstants>$(DefineConstants);TRACE</DefineConstants>
-    <DebugSymbols>true</DebugSymbols>
-    <DebugType>full</DebugType>
-  </PropertyGroup>
-  <PropertyGroup Condition="'$(Configuration)|$(Platform)' == 'Development|AnyCPU'">
-    <Optimize>true</Optimize>
-    <DefineConstants>$(DefineConstants);TRACE</DefineConstants>
-  </PropertyGroup>
-  <PropertyGroup Condition="'$(Configuration)|$(Platform)' == 'Release|AnyCPU'">
-    <Optimize>true</Optimize>
-    <DefineConstants>$(DefineConstants);TRACE</DefineConstants>
-  </PropertyGroup>
 
   <!-- The following will glob all platform extensions' AutomationUtils source
   files and include them in the build. They will also appear in VisualStudio -->
diff --git a/Engine/Source/Programs/AutomationTool/BuildGraph/BuildGraph.Automation.csproj b/Engine/Source/Programs/AutomationTool/BuildGraph/BuildGraph.Automation.csproj
--- a/Engine/Source/Programs/AutomationTool/BuildGraph/BuildGraph.Automation.csproj
+++ b/Engine/Source/Programs/AutomationTool/BuildGraph/BuildGraph.Automation.csproj
@@ -2,40 +2,14 @@
   <Import Project="..\..\Shared\UnrealEngine.csproj.props" />
 
   <PropertyGroup>
-    <TargetFramework>net8.0</TargetFramework>
-    <Configuration Condition=" '$(Configuration)' == '' ">Development</Configuration>
     <OutputType>Library</OutputType>
-    <AppendTargetFrameworkToOutputPath>false</AppendTargetFrameworkToOutputPath>
     <GenerateAssemblyInfo>false</GenerateAssemblyInfo>
     <GenerateTargetFrameworkAttribute>false</GenerateTargetFrameworkAttribute>
-    <Configurations>Debug;Release;Development;Analyze</Configurations>
     <RootNamespace>AutomationTool</RootNamespace>
     <AssemblyName>BuildGraph.Automation</AssemblyName>
     <WarningsNotAsErrors>612,618</WarningsNotAsErrors>
     <OutputPath>..\..\..\..\Binaries\DotNET\AutomationTool\AutomationScripts\BuildGraph</OutputPath>
     <TreatWarningsAsErrors>true</TreatWarningsAsErrors>
-    <DebugType>pdbonly</DebugType>
-    <DebugType Condition="'$([System.Runtime.InteropServices.RuntimeInformation]::IsOSPlatform($([System.Runtime.InteropServices.OSPlatform]::Windows)))' == 'true' And '$([System.Runtime.InteropServices.RuntimeInformation]::OSArchitecture)' == 'Arm64'">portable</DebugType>
-    <GenerateDocumentationFile>True</GenerateDocumentationFile>
-  </PropertyGroup>
-  <PropertyGroup Condition=" '$(Configuration)|$(Platform)' == 'Development|AnyCPU' ">
-    <DefineConstants>$(DefineConstants);TRACE</DefineConstants>
-    <Optimize>true</Optimize>
-  </PropertyGroup>
-  <PropertyGroup Condition=" '$(Configuration)|$(Platform)' == 'Release|AnyCPU' ">
-    <DefineConstants>$(DefineConstants);TRACE</DefineConstants>
-    <Optimize>true</Optimize>
-  </PropertyGroup>
-  <PropertyGroup Condition=" '$(Configuration)|$(Platform)' == 'Debug|AnyCPU' ">
-    <DefineConstants>$(DefineConstants);DEBUG;TRACE</DefineConstants>
-    <DebugSymbols>true</DebugSymbols>
-    <DebugType>full</DebugType>
-  </PropertyGroup>
-  <PropertyGroup Condition="'$(Configuration)|$(Platform)' == 'Analyze|AnyCPU'">
-    <DocumentationFile></DocumentationFile>
-    <TreatWarningsAsErrors>false</TreatWarningsAsErrors>
-    <RunAnalyzersDuringBuild>True</RunAnalyzersDuringBuild>
-    <EnforceCodeStyleInBuild>True</EnforceCodeStyleInBuild>
     <GenerateDocumentationFile>True</GenerateDocumentationFile>
   </PropertyGroup>
   <ItemGroup>
diff --git a/Engine/Source/Programs/AutomationTool/CookedEditor/CookedEditor.Automation.csproj b/Engine/Source/Programs/AutomationTool/CookedEditor/CookedEditor.Automation.csproj
--- a/Engine/Source/Programs/AutomationTool/CookedEditor/CookedEditor.Automation.csproj
+++ b/Engine/Source/Programs/AutomationTool/CookedEditor/CookedEditor.Automation.csproj
@@ -2,32 +2,14 @@
   <Import Project="..\..\Shared\UnrealEngine.csproj.props" />
   
   <PropertyGroup>
-    <TargetFramework>net8.0</TargetFramework>
-    <Configuration Condition=" '$(Configuration)' == '' ">Development</Configuration>
     <OutputType>Library</OutputType>
-    <AppendTargetFrameworkToOutputPath>false</AppendTargetFrameworkToOutputPath>
     <GenerateAssemblyInfo>false</GenerateAssemblyInfo>
     <GenerateTargetFrameworkAttribute>false</GenerateTargetFrameworkAttribute>
-    <Configurations>Debug;Release;Development</Configurations>
     <RootNamespace>AutomationTool</RootNamespace>
     <AssemblyName>CookedEditor.Automation</AssemblyName>
     <WarningsNotAsErrors>612,618</WarningsNotAsErrors>
     <OutputPath>..\..\..\..\Binaries\DotNET\AutomationTool\AutomationScripts\CookedEditor</OutputPath>
     <TreatWarningsAsErrors>true</TreatWarningsAsErrors>
-    <DebugType>pdbonly</DebugType>
-  </PropertyGroup>
-  <PropertyGroup Condition=" '$(Configuration)|$(Platform)' == 'Development|AnyCPU' ">
-    <DefineConstants>$(DefineConstants);TRACE</DefineConstants>
-    <Optimize>true</Optimize>
-  </PropertyGroup>
-  <PropertyGroup Condition=" '$(Configuration)|$(Platform)' == 'Release|AnyCPU' ">
-    <DefineConstants>$(DefineConstants);TRACE</DefineConstants>
-    <Optimize>true</Optimize>
-  </PropertyGroup>
-  <PropertyGroup Condition=" '$(Configuration)|$(Platform)' == 'Debug|AnyCPU' ">
-    <DefineConstants>$(DefineConstants);DEBUG;TRACE</DefineConstants>
-    <DebugSymbols>true</DebugSymbols>
-    <DebugType>full</DebugType>
   </PropertyGroup>
   <ItemGroup>
     <Compile Include="..\..\Shared\MetaData.cs">
diff --git a/Engine/Source/Programs/AutomationTool/CrowdinLocalization/CrowdinLocalization.Automation.csproj b/Engine/Source/Programs/AutomationTool/CrowdinLocalization/CrowdinLocalization.Automation.csproj
--- a/Engine/Source/Programs/AutomationTool/CrowdinLocalization/CrowdinLocalization.Automation.csproj
+++ b/Engine/Source/Programs/AutomationTool/CrowdinLocalization/CrowdinLocalization.Automation.csproj
@@ -2,37 +2,17 @@
   <Import Project="..\..\Shared\UnrealEngine.csproj.props" />
   
   <PropertyGroup>
-    <TargetFramework>net8.0</TargetFramework>
-    <Configuration Condition=" '$(Configuration)' == '' ">Development</Configuration>
     <OutputType>Library</OutputType>
-    <AppendTargetFrameworkToOutputPath>false</AppendTargetFrameworkToOutputPath>
     <GenerateAssemblyInfo>false</GenerateAssemblyInfo>
     <GenerateTargetFrameworkAttribute>false</GenerateTargetFrameworkAttribute>
-    <Configurations>Debug;Release;Development</Configurations>
     <RootNamespace>AutomationTool</RootNamespace>
     <AssemblyName>CrowdinLocalization.Automation</AssemblyName>
     <WarningsNotAsErrors>612,618</WarningsNotAsErrors>
     <OutputPath>..\..\..\..\Binaries\DotNET\AutomationTool\AutomationScripts\CrowdinLocalization</OutputPath>
+    <NoWarn>$(NoWarn);1701;1702;CS1998;CS0649</NoWarn>
     <TreatWarningsAsErrors>true</TreatWarningsAsErrors>
     <CopyLocalLockFileAssemblies>true</CopyLocalLockFileAssemblies>
     <SatelliteResourceLanguages>en</SatelliteResourceLanguages> <!-- remove non english resource languages -->
-    <DebugType>pdbonly</DebugType>
-  </PropertyGroup>
-  <PropertyGroup Condition=" '$(Configuration)|$(Platform)' == 'Development|AnyCPU' ">
-    <DefineConstants>$(DefineConstants);TRACE</DefineConstants>
-    <Optimize>true</Optimize>
-    <NoWarn>$(NoWarn);1701;1702;CS1998;CS0649</NoWarn>
-  </PropertyGroup>
-  <PropertyGroup Condition=" '$(Configuration)|$(Platform)' == 'Release|AnyCPU' ">
-    <DefineConstants>$(DefineConstants);TRACE</DefineConstants>
-    <Optimize>true</Optimize>
-    <NoWarn>$(NoWarn);1701;1702;CS1998;CS0649</NoWarn>
-  </PropertyGroup>
-  <PropertyGroup Condition=" '$(Configuration)|$(Platform)' == 'Debug|AnyCPU' ">
-    <DefineConstants>$(DefineConstants);DEBUG;TRACE</DefineConstants>
-    <DebugSymbols>true</DebugSymbols>
-    <DebugType>full</DebugType>
-    <NoWarn>$(NoWarn);1701;1702;CS1998;CS0649</NoWarn>
   </PropertyGroup>
 
   <ItemGroup>
diff --git a/Engine/Source/Programs/AutomationTool/Gauntlet/Gauntlet.Automation.csproj b/Engine/Source/Programs/AutomationTool/Gauntlet/Gauntlet.Automation.csproj
--- a/Engine/Source/Programs/AutomationTool/Gauntlet/Gauntlet.Automation.csproj
+++ b/Engine/Source/Programs/AutomationTool/Gauntlet/Gauntlet.Automation.csproj
@@ -2,35 +2,17 @@
   <Import Project="..\..\Shared\UnrealEngine.csproj.props" />
   
   <PropertyGroup>
-    <TargetFramework>net8.0</TargetFramework>
-    <Configuration Condition=" '$(Configuration)' == '' ">Development</Configuration>
     <OutputType>Library</OutputType>
-    <AppendTargetFrameworkToOutputPath>false</AppendTargetFrameworkToOutputPath>
     <GenerateAssemblyInfo>false</GenerateAssemblyInfo>
     <GenerateTargetFrameworkAttribute>false</GenerateTargetFrameworkAttribute>
-    <Configurations>Debug;Release;Development</Configurations>
     <RootNamespace>Gauntlet.Automation</RootNamespace>
     <AssemblyName>Gauntlet.Automation</AssemblyName>
     <WarningsNotAsErrors>612,618</WarningsNotAsErrors>
     <OutputPath>..\..\..\..\Binaries\DotNET\AutomationTool\AutomationScripts\Gauntlet</OutputPath>
     <TreatWarningsAsErrors>true</TreatWarningsAsErrors>
-    <DebugType>pdbonly</DebugType>
     <CopyLocalLockFileAssemblies>true</CopyLocalLockFileAssemblies>
     <SatelliteResourceLanguages>en</SatelliteResourceLanguages>    <!-- remove non english resource languages -->
   </PropertyGroup>
-  <PropertyGroup Condition=" '$(Configuration)|$(Platform)' == 'Development|AnyCPU' ">
-    <DefineConstants>$(DefineConstants);TRACE</DefineConstants>
-    <Optimize>true</Optimize>
-  </PropertyGroup>
-  <PropertyGroup Condition=" '$(Configuration)|$(Platform)' == 'Release|AnyCPU' ">
-    <DefineConstants>$(DefineConstants);TRACE</DefineConstants>
-    <Optimize>true</Optimize>
-  </PropertyGroup>
-  <PropertyGroup Condition=" '$(Configuration)|$(Platform)' == 'Debug|AnyCPU' ">
-    <DefineConstants>$(DefineConstants);DEBUG;TRACE</DefineConstants>
-    <DebugSymbols>true</DebugSymbols>
-    <DebugType>full</DebugType>
-  </PropertyGroup>
   <ItemGroup>
     <Reference Include="fastJSON">
       <HintPath>..\..\..\..\Binaries\ThirdParty\fastJSON\netstandard2.0\fastJSON.dll</HintPath>
diff --git a/Engine/Source/Programs/AutomationTool/IOS/IOS.Automation.csproj b/Engine/Source/Programs/AutomationTool/IOS/IOS.Automation.csproj
--- a/Engine/Source/Programs/AutomationTool/IOS/IOS.Automation.csproj
+++ b/Engine/Source/Programs/AutomationTool/IOS/IOS.Automation.csproj
@@ -2,32 +2,14 @@
   <Import Project="..\..\Shared\UnrealEngine.csproj.props" />
   
   <PropertyGroup>
-    <TargetFramework>net8.0</TargetFramework>
-    <Configuration Condition=" '$(Configuration)' == '' ">Development</Configuration>
     <OutputType>Library</OutputType>
-    <AppendTargetFrameworkToOutputPath>false</AppendTargetFrameworkToOutputPath>
     <GenerateAssemblyInfo>false</GenerateAssemblyInfo>
     <GenerateTargetFrameworkAttribute>false</GenerateTargetFrameworkAttribute>
-    <Configurations>Debug;Release;Development</Configurations>
     <RootNamespace>IOS.Automation</RootNamespace>
     <AssemblyName>IOS.Automation</AssemblyName>
     <WarningsNotAsErrors>612,618</WarningsNotAsErrors>
     <OutputPath>..\..\..\..\Binaries\DotNET\AutomationTool\AutomationScripts\Platforms\IOS</OutputPath>
     <TreatWarningsAsErrors>true</TreatWarningsAsErrors>
-    <DebugType>pdbonly</DebugType>
-  </PropertyGroup>
-  <PropertyGroup Condition=" '$(Configuration)|$(Platform)' == 'Development|AnyCPU' ">
-    <DefineConstants>$(DefineConstants);TRACE</DefineConstants>
-    <Optimize>true</Optimize>
-  </PropertyGroup>
-  <PropertyGroup Condition=" '$(Configuration)|$(Platform)' == 'Release|AnyCPU' ">
-    <DefineConstants>$(DefineConstants);TRACE</DefineConstants>
-    <Optimize>true</Optimize>
-  </PropertyGroup>
-  <PropertyGroup Condition=" '$(Configuration)|$(Platform)' == 'Debug|AnyCPU' ">
-    <DefineConstants>$(DefineConstants);DEBUG;TRACE</DefineConstants>
-    <DebugSymbols>true</DebugSymbols>
-    <DebugType>full</DebugType>
   </PropertyGroup>
   <ItemGroup>
     <!-- BouncyCastle.Crypto is not actually used but is required so that IPhonePackager has it available after we have copied files to a local folder -->
diff --git a/Engine/Source/Programs/AutomationTool/Insights/Insights.Automation.csproj b/Engine/Source/Programs/AutomationTool/Insights/Insights.Automation.csproj
--- a/Engine/Source/Programs/AutomationTool/Insights/Insights.Automation.csproj
+++ b/Engine/Source/Programs/AutomationTool/Insights/Insights.Automation.csproj
@@ -2,32 +2,14 @@
   <Import Project="..\..\Shared\UnrealEngine.csproj.props" />
   
   <PropertyGroup>
-    <TargetFramework>net8.0</TargetFramework>
-	<Configuration Condition=" '$(Configuration)' == '' ">Development</Configuration>
     <OutputType>Library</OutputType>
-	<RootNamespace>Insights.Automation</RootNamespace>
+    <RootNamespace>Insights.Automation</RootNamespace>
     <AssemblyName>Insights.Automation</AssemblyName>
     <GenerateAssemblyInfo>false</GenerateAssemblyInfo>
     <GenerateTargetFrameworkAttribute>false</GenerateTargetFrameworkAttribute>
     <OutputPath>..\..\..\..\Binaries\DotNET\AutomationTool\AutomationScripts\Insights</OutputPath>
-    <Configurations>Debug;Release;Development</Configurations>
     <TreatWarningsAsErrors>true</TreatWarningsAsErrors>
-    <AppendTargetFrameworkToOutputPath>false</AppendTargetFrameworkToOutputPath>
-    <DebugType>pdbonly</DebugType>
-	<SatelliteResourceLanguages>en</SatelliteResourceLanguages>    <!-- remove non english resource languages -->
-  </PropertyGroup>
-  <PropertyGroup Condition=" '$(Configuration)|$(Platform)' == 'Development|AnyCPU' ">
-    <DefineConstants>$(DefineConstants);TRACE</DefineConstants>
-    <Optimize>true</Optimize>
-  </PropertyGroup>
-  <PropertyGroup Condition=" '$(Configuration)|$(Platform)' == 'Release|AnyCPU' ">
-    <DefineConstants>$(DefineConstants);TRACE</DefineConstants>
-    <Optimize>true</Optimize>
-  </PropertyGroup>
-  <PropertyGroup Condition=" '$(Configuration)|$(Platform)' == 'Debug|AnyCPU' ">
-    <DefineConstants>$(DefineConstants);DEBUG;TRACE</DefineConstants>
-    <DebugSymbols>true</DebugSymbols>
-    <DebugType>full</DebugType>
+    <SatelliteResourceLanguages>en</SatelliteResourceLanguages>    <!-- remove non english resource languages -->
   </PropertyGroup>
   <ItemGroup>
     <ProjectReference Include="..\..\Shared\EpicGames.Core\EpicGames.Core.csproj" PrivateAssets="All">
diff --git a/Engine/Source/Programs/AutomationTool/Linux/Linux.Automation.csproj b/Engine/Source/Programs/AutomationTool/Linux/Linux.Automation.csproj
--- a/Engine/Source/Programs/AutomationTool/Linux/Linux.Automation.csproj
+++ b/Engine/Source/Programs/AutomationTool/Linux/Linux.Automation.csproj
@@ -2,32 +2,14 @@
   <Import Project="..\..\Shared\UnrealEngine.csproj.props" />
   
   <PropertyGroup>
-    <TargetFramework>net8.0</TargetFramework>
-    <Configuration Condition=" '$(Configuration)' == '' ">Development</Configuration>
     <OutputType>Library</OutputType>
-    <AppendTargetFrameworkToOutputPath>false</AppendTargetFrameworkToOutputPath>
     <GenerateAssemblyInfo>false</GenerateAssemblyInfo>
     <GenerateTargetFrameworkAttribute>false</GenerateTargetFrameworkAttribute>
-    <Configurations>Debug;Release;Development</Configurations>
     <RootNamespace>Linux.Automation</RootNamespace>
     <AssemblyName>Linux.Automation</AssemblyName>
     <WarningsNotAsErrors>612,618</WarningsNotAsErrors>
     <OutputPath>..\..\..\..\Binaries\DotNET\AutomationTool\AutomationScripts\Platforms\Linux\</OutputPath>
     <TreatWarningsAsErrors>true</TreatWarningsAsErrors>
-    <DebugType>pdbonly</DebugType>
-  </PropertyGroup>
-  <PropertyGroup Condition=" '$(Configuration)|$(Platform)' == 'Development|AnyCPU' ">
-    <DefineConstants>$(DefineConstants);TRACE</DefineConstants>
-    <Optimize>true</Optimize>
-  </PropertyGroup>
-  <PropertyGroup Condition=" '$(Configuration)|$(Platform)' == 'Release|AnyCPU' ">
-    <DefineConstants>$(DefineConstants);TRACE</DefineConstants>
-    <Optimize>true</Optimize>
-  </PropertyGroup>
-  <PropertyGroup Condition=" '$(Configuration)|$(Platform)' == 'Debug|AnyCPU' ">
-    <DefineConstants>$(DefineConstants);DEBUG;TRACE</DefineConstants>
-    <DebugSymbols>true</DebugSymbols>
-    <DebugType>full</DebugType>
   </PropertyGroup>
 
   <ItemGroup>
diff --git a/Engine/Source/Programs/AutomationTool/LiveLinkHub/LiveLinkHub.Automation.csproj b/Engine/Source/Programs/AutomationTool/LiveLinkHub/LiveLinkHub.Automation.csproj
--- a/Engine/Source/Programs/AutomationTool/LiveLinkHub/LiveLinkHub.Automation.csproj
+++ b/Engine/Source/Programs/AutomationTool/LiveLinkHub/LiveLinkHub.Automation.csproj
@@ -2,12 +2,9 @@
   <Import Project="..\..\Shared\UnrealEngine.csproj.props" />
 
   <PropertyGroup>
-    <TargetFramework>net8.0</TargetFramework>
     <OutputType>Library</OutputType>
-    <AppendTargetFrameworkToOutputPath>false</AppendTargetFrameworkToOutputPath>
     <GenerateAssemblyInfo>false</GenerateAssemblyInfo>
     <GenerateTargetFrameworkAttribute>false</GenerateTargetFrameworkAttribute>
-    <Configurations>Debug;Release;Development</Configurations>
     <OutputPath>..\..\..\..\Binaries\DotNET\AutomationTool\AutomationScripts\LiveLinkHub</OutputPath>
   </PropertyGroup>
 
diff --git a/Engine/Source/Programs/AutomationTool/Localization/Localization.Automation.csproj b/Engine/Source/Programs/AutomationTool/Localization/Localization.Automation.csproj
--- a/Engine/Source/Programs/AutomationTool/Localization/Localization.Automation.csproj
+++ b/Engine/Source/Programs/AutomationTool/Localization/Localization.Automation.csproj
@@ -2,32 +2,14 @@
   <Import Project="..\..\Shared\UnrealEngine.csproj.props" />
   
   <PropertyGroup>
-    <TargetFramework>net8.0</TargetFramework>
-    <Configuration Condition=" '$(Configuration)' == '' ">Development</Configuration>
     <OutputType>Library</OutputType>
-    <AppendTargetFrameworkToOutputPath>false</AppendTargetFrameworkToOutputPath>
     <GenerateAssemblyInfo>false</GenerateAssemblyInfo>
     <GenerateTargetFrameworkAttribute>false</GenerateTargetFrameworkAttribute>
-    <Configurations>Debug;Release;Development</Configurations>
     <RootNamespace>AutomationTool</RootNamespace>
     <AssemblyName>Localization.Automation</AssemblyName>
     <WarningsNotAsErrors>612,618</WarningsNotAsErrors>
     <OutputPath>..\..\..\..\Binaries\DotNET\AutomationTool\AutomationScripts\Localization</OutputPath>
     <TreatWarningsAsErrors>true</TreatWarningsAsErrors>
-    <DebugType>pdbonly</DebugType>
-  </PropertyGroup>
-  <PropertyGroup Condition=" '$(Configuration)|$(Platform)' == 'Development|AnyCPU' ">
-    <DefineConstants>$(DefineConstants);TRACE</DefineConstants>
-    <Optimize>true</Optimize>
-  </PropertyGroup>
-  <PropertyGroup Condition=" '$(Configuration)|$(Platform)' == 'Release|AnyCPU' ">
-    <DefineConstants>$(DefineConstants);TRACE</DefineConstants>
-    <Optimize>true</Optimize>
-  </PropertyGroup>
-  <PropertyGroup Condition=" '$(Configuration)|$(Platform)' == 'Debug|AnyCPU' ">
-    <DefineConstants>$(DefineConstants);DEBUG;TRACE</DefineConstants>
-    <DebugSymbols>true</DebugSymbols>
-    <DebugType>full</DebugType>
   </PropertyGroup>
 
   <ItemGroup>
diff --git a/Engine/Source/Programs/AutomationTool/LowLevelTests/LowLevelTests.Automation.csproj b/Engine/Source/Programs/AutomationTool/LowLevelTests/LowLevelTests.Automation.csproj
--- a/Engine/Source/Programs/AutomationTool/LowLevelTests/LowLevelTests.Automation.csproj
+++ b/Engine/Source/Programs/AutomationTool/LowLevelTests/LowLevelTests.Automation.csproj
@@ -2,32 +2,14 @@
   <Import Project="..\..\Shared\UnrealEngine.csproj.props" />
   
   <PropertyGroup>
-    <TargetFramework>net8.0</TargetFramework>
-	<Configuration Condition=" '$(Configuration)' == '' ">Development</Configuration>
     <OutputType>Library</OutputType>
-	<RootNamespace>LowLevelTests.Automation</RootNamespace>
+    <RootNamespace>LowLevelTests.Automation</RootNamespace>
     <AssemblyName>LowLevelTests.Automation</AssemblyName>
     <GenerateAssemblyInfo>false</GenerateAssemblyInfo>
     <GenerateTargetFrameworkAttribute>false</GenerateTargetFrameworkAttribute>
     <OutputPath>..\..\..\..\Binaries\DotNET\AutomationTool\AutomationScripts\LowLevelTests</OutputPath>
-    <Configurations>Debug;Release;Development</Configurations>
     <TreatWarningsAsErrors>true</TreatWarningsAsErrors>
-    <AppendTargetFrameworkToOutputPath>false</AppendTargetFrameworkToOutputPath>
-    <DebugType>pdbonly</DebugType>
-	<SatelliteResourceLanguages>en</SatelliteResourceLanguages>    <!-- remove non english resource languages -->
-  </PropertyGroup>
-  <PropertyGroup Condition=" '$(Configuration)|$(Platform)' == 'Development|AnyCPU' ">
-    <DefineConstants>$(DefineConstants);TRACE</DefineConstants>
-    <Optimize>true</Optimize>
-  </PropertyGroup>
-  <PropertyGroup Condition=" '$(Configuration)|$(Platform)' == 'Release|AnyCPU' ">
-    <DefineConstants>$(DefineConstants);TRACE</DefineConstants>
-    <Optimize>true</Optimize>
-  </PropertyGroup>
-  <PropertyGroup Condition=" '$(Configuration)|$(Platform)' == 'Debug|AnyCPU' ">
-    <DefineConstants>$(DefineConstants);DEBUG;TRACE</DefineConstants>
-    <DebugSymbols>true</DebugSymbols>
-    <DebugType>full</DebugType>
+    <SatelliteResourceLanguages>en</SatelliteResourceLanguages>    <!-- remove non english resource languages -->
   </PropertyGroup>
   <ItemGroup>
     <ProjectReference Include="..\..\Shared\EpicGames.Core\EpicGames.Core.csproj" PrivateAssets="All">
diff --git a/Engine/Source/Programs/AutomationTool/Mac/Mac.Automation.csproj b/Engine/Source/Programs/AutomationTool/Mac/Mac.Automation.csproj
--- a/Engine/Source/Programs/AutomationTool/Mac/Mac.Automation.csproj
+++ b/Engine/Source/Programs/AutomationTool/Mac/Mac.Automation.csproj
@@ -2,32 +2,14 @@
   <Import Project="..\..\Shared\UnrealEngine.csproj.props" />
   
   <PropertyGroup>
-    <TargetFramework>net8.0</TargetFramework>
-    <Configuration Condition=" '$(Configuration)' == '' ">Development</Configuration>
     <OutputType>Library</OutputType>
-    <AppendTargetFrameworkToOutputPath>false</AppendTargetFrameworkToOutputPath>
     <GenerateAssemblyInfo>false</GenerateAssemblyInfo>
     <GenerateTargetFrameworkAttribute>false</GenerateTargetFrameworkAttribute>
-    <Configurations>Debug;Release;Development</Configurations>
     <RootNamespace>Mac.Automation</RootNamespace>
     <AssemblyName>Mac.Automation</AssemblyName>
     <WarningsNotAsErrors>612,618</WarningsNotAsErrors>
     <OutputPath>..\..\..\..\Binaries\DotNET\AutomationTool\AutomationScripts\Platforms\Mac\</OutputPath>
     <TreatWarningsAsErrors>true</TreatWarningsAsErrors>
-    <DebugType>pdbonly</DebugType>
-  </PropertyGroup>
-  <PropertyGroup Condition=" '$(Configuration)|$(Platform)' == 'Development|AnyCPU' ">
-    <DefineConstants>$(DefineConstants);TRACE</DefineConstants>
-    <Optimize>true</Optimize>
-  </PropertyGroup>
-  <PropertyGroup Condition=" '$(Configuration)|$(Platform)' == 'Release|AnyCPU' ">
-    <DefineConstants>$(DefineConstants);TRACE</DefineConstants>
-    <Optimize>true</Optimize>
-  </PropertyGroup>
-  <PropertyGroup Condition=" '$(Configuration)|$(Platform)' == 'Debug|AnyCPU' ">
-    <DefineConstants>$(DefineConstants);DEBUG;TRACE</DefineConstants>
-    <DebugSymbols>true</DebugSymbols>
-    <DebugType>full</DebugType>
   </PropertyGroup>
   <ItemGroup>
     <Compile Include="..\..\Shared\MetaData.cs">
diff --git a/Engine/Source/Programs/AutomationTool/Microsoft/Microsoft.Automation.csproj b/Engine/Source/Programs/AutomationTool/Microsoft/Microsoft.Automation.csproj
--- a/Engine/Source/Programs/AutomationTool/Microsoft/Microsoft.Automation.csproj
+++ b/Engine/Source/Programs/AutomationTool/Microsoft/Microsoft.Automation.csproj
@@ -2,32 +2,14 @@
   <Import Project="..\..\Shared\UnrealEngine.csproj.props" />
   
   <PropertyGroup>
-    <TargetFramework>net8.0</TargetFramework>
-    <Configuration Condition=" '$(Configuration)' == '' ">Development</Configuration>
     <OutputType>Library</OutputType>
-    <AppendTargetFrameworkToOutputPath>false</AppendTargetFrameworkToOutputPath>
     <GenerateAssemblyInfo>false</GenerateAssemblyInfo>
     <GenerateTargetFrameworkAttribute>false</GenerateTargetFrameworkAttribute>
-    <Configurations>Debug;Release;Development</Configurations>
     <RootNamespace>Microsoft.Automation</RootNamespace>
     <AssemblyName>Microsoft.Automation</AssemblyName>
     <WarningsNotAsErrors>612,618</WarningsNotAsErrors>
     <OutputPath>..\..\..\..\Binaries\DotNET\AutomationTool\AutomationScripts\Platforms\Microsoft</OutputPath>
     <TreatWarningsAsErrors>true</TreatWarningsAsErrors>
-    <DebugType>pdbonly</DebugType>
-  </PropertyGroup>
-  <PropertyGroup Condition=" '$(Configuration)|$(Platform)' == 'Development|AnyCPU' ">
-    <DefineConstants>$(DefineConstants);TRACE</DefineConstants>
-    <Optimize>true</Optimize>
-  </PropertyGroup>
-  <PropertyGroup Condition=" '$(Configuration)|$(Platform)' == 'Release|AnyCPU' ">
-    <DefineConstants>$(DefineConstants);TRACE</DefineConstants>
-    <Optimize>true</Optimize>
-  </PropertyGroup>
-  <PropertyGroup Condition=" '$(Configuration)|$(Platform)' == 'Debug|AnyCPU' ">
-    <DefineConstants>$(DefineConstants);DEBUG;TRACE</DefineConstants>
-    <DebugSymbols>true</DebugSymbols>
-    <DebugType>full</DebugType>
   </PropertyGroup>
   <ItemGroup>
     <Compile Include="..\..\Shared\MetaData.cs">
diff --git a/Engine/Source/Programs/AutomationTool/Mutable/Mutable.Automation.csproj b/Engine/Source/Programs/AutomationTool/Mutable/Mutable.Automation.csproj
--- a/Engine/Source/Programs/AutomationTool/Mutable/Mutable.Automation.csproj
+++ b/Engine/Source/Programs/AutomationTool/Mutable/Mutable.Automation.csproj
@@ -2,16 +2,12 @@
   <Import Project="..\..\Shared\UnrealEngine.csproj.props" />
 
   <PropertyGroup>
-    <TargetFramework>net8.0</TargetFramework>
-	<OutputType>Library</OutputType>
+    <OutputType>Library</OutputType>
     <ImplicitUsings>enable</ImplicitUsings>
     <Nullable>enable</Nullable>
-    <Configurations>Debug;Development</Configurations>
-    <!-- <BaseOutputPath>bin\</BaseOutputPath> -->
-	<AppendTargetFrameworkToOutputPath>false</AppendTargetFrameworkToOutputPath>
     <GenerateAssemblyInfo>false</GenerateAssemblyInfo>
     <GenerateTargetFrameworkAttribute>false</GenerateTargetFrameworkAttribute>
-	<OutputPath>..\..\..\..\Binaries\DotNET\AutomationTool\AutomationScripts\Mutable\</OutputPath>
+    <OutputPath>..\..\..\..\Binaries\DotNET\AutomationTool\AutomationScripts\Mutable\</OutputPath>
   </PropertyGroup>
 
   <ItemGroup>
diff --git a/Engine/Source/Programs/AutomationTool/OneSkyLocalization/OneSkyLocalization.Automation.csproj b/Engine/Source/Programs/AutomationTool/OneSkyLocalization/OneSkyLocalization.Automation.csproj
--- a/Engine/Source/Programs/AutomationTool/OneSkyLocalization/OneSkyLocalization.Automation.csproj
+++ b/Engine/Source/Programs/AutomationTool/OneSkyLocalization/OneSkyLocalization.Automation.csproj
@@ -2,32 +2,14 @@
   <Import Project="..\..\Shared\UnrealEngine.csproj.props" />
   
   <PropertyGroup>
-    <TargetFramework>net8.0</TargetFramework>
-    <Configuration Condition=" '$(Configuration)' == '' ">Development</Configuration>
     <OutputType>Library</OutputType>
-    <AppendTargetFrameworkToOutputPath>false</AppendTargetFrameworkToOutputPath>
     <GenerateAssemblyInfo>false</GenerateAssemblyInfo>
     <GenerateTargetFrameworkAttribute>false</GenerateTargetFrameworkAttribute>
-    <Configurations>Debug;Release;Development</Configurations>
     <RootNamespace>AutomationTool</RootNamespace>
     <AssemblyName>OneSkyLocalization.Automation</AssemblyName>
     <WarningsNotAsErrors>612,618</WarningsNotAsErrors>
     <OutputPath>..\..\..\..\Binaries\DotNET\AutomationTool\AutomationScripts\OneSkyLocalization</OutputPath>
     <TreatWarningsAsErrors>true</TreatWarningsAsErrors>
-    <DebugType>pdbonly</DebugType>
-  </PropertyGroup>
-  <PropertyGroup Condition=" '$(Configuration)|$(Platform)' == 'Development|AnyCPU' ">
-    <DefineConstants>$(DefineConstants);TRACE</DefineConstants>
-    <Optimize>true</Optimize>
-  </PropertyGroup>
-  <PropertyGroup Condition=" '$(Configuration)|$(Platform)' == 'Release|AnyCPU' ">
-    <DefineConstants>$(DefineConstants);TRACE</DefineConstants>
-    <Optimize>true</Optimize>
-  </PropertyGroup>
-  <PropertyGroup Condition=" '$(Configuration)|$(Platform)' == 'Debug|AnyCPU' ">
-    <DefineConstants>$(DefineConstants);DEBUG;TRACE</DefineConstants>
-    <DebugSymbols>true</DebugSymbols>
-    <DebugType>full</DebugType>
   </PropertyGroup>
   <ItemGroup>
     <Reference Include="OneSky, Version=1.0.0.0, Culture=neutral, processorArchitecture=MSIL">
diff --git a/Engine/Source/Programs/AutomationTool/Scripts/AutomationScripts.Automation.csproj b/Engine/Source/Programs/AutomationTool/Scripts/AutomationScripts.Automation.csproj
--- a/Engine/Source/Programs/AutomationTool/Scripts/AutomationScripts.Automation.csproj
+++ b/Engine/Source/Programs/AutomationTool/Scripts/AutomationScripts.Automation.csproj
@@ -2,32 +2,14 @@
   <Import Project="..\..\Shared\UnrealEngine.csproj.props" />
   
   <PropertyGroup>
-    <TargetFramework>net8.0</TargetFramework>
-    <Configuration Condition=" '$(Configuration)' == '' ">Development</Configuration>
     <OutputType>Library</OutputType>
-    <AppendTargetFrameworkToOutputPath>false</AppendTargetFrameworkToOutputPath>
     <GenerateAssemblyInfo>false</GenerateAssemblyInfo>
     <GenerateTargetFrameworkAttribute>false</GenerateTargetFrameworkAttribute>
-    <Configurations>Debug;Release;Development</Configurations>
     <RootNamespace>AutomationTool</RootNamespace>
     <AssemblyName>AutomationScripts.Automation</AssemblyName>
     <WarningsNotAsErrors>612,618</WarningsNotAsErrors>
     <OutputPath>..\..\..\..\Binaries\DotNET\AutomationTool\AutomationScripts\Scripts</OutputPath>
     <TreatWarningsAsErrors>true</TreatWarningsAsErrors>
-    <DebugType>pdbonly</DebugType>
-  </PropertyGroup>
-  <PropertyGroup Condition=" '$(Configuration)|$(Platform)' == 'Development|AnyCPU' ">
-    <DefineConstants>$(DefineConstants);TRACE</DefineConstants>
-    <Optimize>true</Optimize>
-  </PropertyGroup>
-  <PropertyGroup Condition=" '$(Configuration)|$(Platform)' == 'Release|AnyCPU' ">
-    <DefineConstants>$(DefineConstants);TRACE</DefineConstants>
-    <Optimize>true</Optimize>
-  </PropertyGroup>
-  <PropertyGroup Condition=" '$(Configuration)|$(Platform)' == 'Debug|AnyCPU' ">
-    <DefineConstants>$(DefineConstants);DEBUG;TRACE</DefineConstants>
-    <DebugSymbols>true</DebugSymbols>
-    <DebugType>full</DebugType>
   </PropertyGroup>
   <ItemGroup>
     <Reference Include="Ionic.Zip.Reduced">
diff --git a/Engine/Source/Programs/AutomationTool/SmartlingLocalization/SmartlingLocalization.Automation.csproj b/Engine/Source/Programs/AutomationTool/SmartlingLocalization/SmartlingLocalization.Automation.csproj
--- a/Engine/Source/Programs/AutomationTool/SmartlingLocalization/SmartlingLocalization.Automation.csproj
+++ b/Engine/Source/Programs/AutomationTool/SmartlingLocalization/SmartlingLocalization.Automation.csproj
@@ -2,13 +2,9 @@
   <Import Project="..\..\Shared\UnrealEngine.csproj.props" />
   
   <PropertyGroup>
-    <TargetFramework>net8.0</TargetFramework>
-    <Configuration Condition=" '$(Configuration)' == '' ">Development</Configuration>
     <OutputType>Library</OutputType>
-    <AppendTargetFrameworkToOutputPath>false</AppendTargetFrameworkToOutputPath>
     <GenerateAssemblyInfo>false</GenerateAssemblyInfo>
     <GenerateTargetFrameworkAttribute>false</GenerateTargetFrameworkAttribute>
-    <Configurations>Debug;Release;Development</Configurations>
     <RootNamespace>AutomationTool</RootNamespace>
     <AssemblyName>SmartlingLocalization.Automation</AssemblyName>
     <WarningsNotAsErrors>612,618</WarningsNotAsErrors>
@@ -16,23 +12,6 @@
     <TreatWarningsAsErrors>true</TreatWarningsAsErrors>
     <CopyLocalLockFileAssemblies>true</CopyLocalLockFileAssemblies>
     <SatelliteResourceLanguages>en</SatelliteResourceLanguages> <!-- remove non english resource languages -->
-    <DebugType>pdbonly</DebugType>
-  </PropertyGroup>
-  <PropertyGroup Condition=" '$(Configuration)|$(Platform)' == 'Development|AnyCPU' ">
-    <DefineConstants>$(DefineConstants);TRACE</DefineConstants>
-    <Optimize>true</Optimize>
-    <NoWarn>$(NoWarn);1701;1702;CS1998;CS0649</NoWarn>
-  </PropertyGroup>
-  <PropertyGroup Condition=" '$(Configuration)|$(Platform)' == 'Release|AnyCPU' ">
-    <DefineConstants>$(DefineConstants);TRACE</DefineConstants>
-    <Optimize>true</Optimize>
-    <NoWarn>$(NoWarn);1701;1702;CS1998;CS0649</NoWarn>
-  </PropertyGroup>
-  <PropertyGroup Condition=" '$(Configuration)|$(Platform)' == 'Debug|AnyCPU' ">
-    <DefineConstants>$(DefineConstants);DEBUG;TRACE</DefineConstants>
-    <DebugSymbols>true</DebugSymbols>
-    <DebugType>full</DebugType>
-    <NoWarn>$(NoWarn);1701;1702;CS1998;CS0649</NoWarn>
   </PropertyGroup>
 
   <ItemGroup>
diff --git a/Engine/Source/Programs/AutomationTool/SteamDeck/SteamDeck.Automation.csproj b/Engine/Source/Programs/AutomationTool/SteamDeck/SteamDeck.Automation.csproj
--- a/Engine/Source/Programs/AutomationTool/SteamDeck/SteamDeck.Automation.csproj
+++ b/Engine/Source/Programs/AutomationTool/SteamDeck/SteamDeck.Automation.csproj
@@ -2,32 +2,14 @@
   <Import Project="..\..\Shared\UnrealEngine.csproj.props" />
   
   <PropertyGroup>
-    <TargetFramework>net8.0</TargetFramework>
-    <Configuration Condition=" '$(Configuration)' == '' ">Development</Configuration>
     <OutputType>Library</OutputType>
-    <AppendTargetFrameworkToOutputPath>false</AppendTargetFrameworkToOutputPath>
     <GenerateAssemblyInfo>false</GenerateAssemblyInfo>
     <GenerateTargetFrameworkAttribute>false</GenerateTargetFrameworkAttribute>
-    <Configurations>Debug;Release;Development</Configurations>
     <RootNamespace>SteamDeck.Automation</RootNamespace>
     <AssemblyName>SteamDeck.Automation</AssemblyName>
     <WarningsNotAsErrors>612,618</WarningsNotAsErrors>
     <OutputPath>..\..\..\..\Binaries\DotNET\AutomationTool\AutomationScripts\Platforms\SteamDeck\</OutputPath>
     <TreatWarningsAsErrors>true</TreatWarningsAsErrors>
-    <DebugType>pdbonly</DebugType>
-  </PropertyGroup>
-  <PropertyGroup Condition=" '$(Configuration)|$(Platform)' == 'Development|AnyCPU' ">
-    <DefineConstants>$(DefineConstants);TRACE</DefineConstants>
-    <Optimize>true</Optimize>
-  </PropertyGroup>
-  <PropertyGroup Condition=" '$(Configuration)|$(Platform)' == 'Release|AnyCPU' ">
-    <DefineConstants>$(DefineConstants);TRACE</DefineConstants>
-    <Optimize>true</Optimize>
-  </PropertyGroup>
-  <PropertyGroup Condition=" '$(Configuration)|$(Platform)' == 'Debug|AnyCPU' ">
-    <DefineConstants>$(DefineConstants);DEBUG;TRACE</DefineConstants>
-    <DebugSymbols>true</DebugSymbols>
-    <DebugType>full</DebugType>
   </PropertyGroup>
 
   <ItemGroup>
diff --git a/Engine/Source/Programs/AutomationTool/TVOS/TVOS.Automation.csproj b/Engine/Source/Programs/AutomationTool/TVOS/TVOS.Automation.csproj
--- a/Engine/Source/Programs/AutomationTool/TVOS/TVOS.Automation.csproj
+++ b/Engine/Source/Programs/AutomationTool/TVOS/TVOS.Automation.csproj
@@ -2,32 +2,14 @@
   <Import Project="..\..\Shared\UnrealEngine.csproj.props" />
   
   <PropertyGroup>
-    <TargetFramework>net8.0</TargetFramework>
-    <Configuration Condition=" '$(Configuration)' == '' ">Development</Configuration>
     <OutputType>Library</OutputType>
-    <AppendTargetFrameworkToOutputPath>false</AppendTargetFrameworkToOutputPath>
     <GenerateAssemblyInfo>false</GenerateAssemblyInfo>
     <GenerateTargetFrameworkAttribute>false</GenerateTargetFrameworkAttribute>
-    <Configurations>Debug;Release;Development</Configurations>
     <RootNamespace>TVOS.Automation</RootNamespace>
     <AssemblyName>TVOS.Automation</AssemblyName>
     <WarningsNotAsErrors>612,618</WarningsNotAsErrors>
     <OutputPath>..\..\..\..\Binaries\DotNET\AutomationTool\AutomationScripts\Platforms\TVOS\</OutputPath>
     <TreatWarningsAsErrors>true</TreatWarningsAsErrors>
-    <DebugType>pdbonly</DebugType>
-  </PropertyGroup>
-  <PropertyGroup Condition=" '$(Configuration)|$(Platform)' == 'Development|AnyCPU' ">
-    <DefineConstants>$(DefineConstants);TRACE</DefineConstants>
-    <Optimize>true</Optimize>
-  </PropertyGroup>
-  <PropertyGroup Condition=" '$(Configuration)|$(Platform)' == 'Release|AnyCPU' ">
-    <DefineConstants>$(DefineConstants);TRACE</DefineConstants>
-    <Optimize>true</Optimize>
-  </PropertyGroup>
-  <PropertyGroup Condition=" '$(Configuration)|$(Platform)' == 'Debug|AnyCPU' ">
-    <DefineConstants>$(DefineConstants);DEBUG;TRACE</DefineConstants>
-    <DebugSymbols>true</DebugSymbols>
-    <DebugType>full</DebugType>
   </PropertyGroup>
   <ItemGroup>
     <Reference Include="Ionic.Zip.Reduced">
diff --git a/Engine/Source/Programs/AutomationTool/Turnkey/Turnkey.Automation.csproj b/Engine/Source/Programs/AutomationTool/Turnkey/Turnkey.Automation.csproj
--- a/Engine/Source/Programs/AutomationTool/Turnkey/Turnkey.Automation.csproj
+++ b/Engine/Source/Programs/AutomationTool/Turnkey/Turnkey.Automation.csproj
@@ -2,35 +2,17 @@
   <Import Project="..\..\Shared\UnrealEngine.csproj.props" />
   
   <PropertyGroup>
-    <TargetFramework>net8.0</TargetFramework>
-    <Configuration Condition=" '$(Configuration)' == '' ">Development</Configuration>
     <OutputType>Library</OutputType>
-    <AppendTargetFrameworkToOutputPath>false</AppendTargetFrameworkToOutputPath>
     <GenerateAssemblyInfo>false</GenerateAssemblyInfo>
     <GenerateTargetFrameworkAttribute>false</GenerateTargetFrameworkAttribute>
-    <Configurations>Debug;Release;Development</Configurations>
     <RootNamespace>Turnkey.Automation</RootNamespace>
     <AssemblyName>Turnkey.Automation</AssemblyName>
     <WarningsNotAsErrors>612,618</WarningsNotAsErrors>
     <OutputPath>..\..\..\..\Binaries\DotNET\AutomationTool\AutomationScripts\Turnkey\</OutputPath>
     <TreatWarningsAsErrors>true</TreatWarningsAsErrors>
-    <DebugType>pdbonly</DebugType>
     <CopyLocalLockFileAssemblies>true</CopyLocalLockFileAssemblies>
     <SatelliteResourceLanguages>en</SatelliteResourceLanguages>    <!-- remove non english resource languages -->
 </PropertyGroup>
-  <PropertyGroup Condition="'$(Configuration)|$(Platform)' == 'Development|AnyCPU' ">
-    <DefineConstants>$(DefineConstants);TRACE</DefineConstants>
-    <Optimize>true</Optimize>
-  </PropertyGroup>
-  <PropertyGroup Condition="'$(Configuration)|$(Platform)' == 'Release|AnyCPU' ">
-    <DefineConstants>$(DefineConstants);TRACE</DefineConstants>
-    <Optimize>true</Optimize>
-  </PropertyGroup>
-  <PropertyGroup Condition="'$(Configuration)|$(Platform)' == 'Debug|AnyCPU' ">
-    <DefineConstants>$(DefineConstants);DEBUG;TRACE</DefineConstants>
-    <DebugSymbols>true</DebugSymbols>
-    <DebugType>full</DebugType>
-  </PropertyGroup>
 
   <ItemGroup>
     <Reference Include="Ionic.Zip.Reduced">
diff --git a/Engine/Source/Programs/AutomationTool/Win/Win.Automation.csproj b/Engine/Source/Programs/AutomationTool/Win/Win.Automation.csproj
--- a/Engine/Source/Programs/AutomationTool/Win/Win.Automation.csproj
+++ b/Engine/Source/Programs/AutomationTool/Win/Win.Automation.csproj
@@ -2,32 +2,14 @@
   <Import Project="..\..\Shared\UnrealEngine.csproj.props" />
   
   <PropertyGroup>
-    <TargetFramework>net8.0</TargetFramework>
-    <Configuration Condition=" '$(Configuration)' == '' ">Development</Configuration>
     <OutputType>Library</OutputType>
-    <AppendTargetFrameworkToOutputPath>false</AppendTargetFrameworkToOutputPath>
     <GenerateAssemblyInfo>false</GenerateAssemblyInfo>
     <GenerateTargetFrameworkAttribute>false</GenerateTargetFrameworkAttribute>
-    <Configurations>Debug;Release;Development</Configurations>
     <RootNamespace>Win.Automation</RootNamespace>
     <AssemblyName>Win.Automation</AssemblyName>
     <WarningsNotAsErrors>612,618</WarningsNotAsErrors>
     <OutputPath>..\..\..\..\Binaries\DotNET\AutomationTool\AutomationScripts\Platforms\Windows</OutputPath>
     <TreatWarningsAsErrors>true</TreatWarningsAsErrors>
-    <DebugType>pdbonly</DebugType>
-  </PropertyGroup>
-  <PropertyGroup Condition=" '$(Configuration)|$(Platform)' == 'Development|AnyCPU' ">
-    <DefineConstants>$(DefineConstants);TRACE</DefineConstants>
-    <Optimize>true</Optimize>
-  </PropertyGroup>
-  <PropertyGroup Condition=" '$(Configuration)|$(Platform)' == 'Release|AnyCPU' ">
-    <DefineConstants>$(DefineConstants);TRACE</DefineConstants>
-    <Optimize>true</Optimize>
-  </PropertyGroup>
-  <PropertyGroup Condition=" '$(Configuration)|$(Platform)' == 'Debug|AnyCPU' ">
-    <DefineConstants>$(DefineConstants);DEBUG;TRACE</DefineConstants>
-    <DebugSymbols>true</DebugSymbols>
-    <DebugType>full</DebugType>
   </PropertyGroup>
   <ItemGroup>
     <Reference Include="Ionic.Zip.Reduced">
diff --git a/Engine/Source/Programs/AutomationTool/XLocLocalization/XLocLocalization.Automation.csproj b/Engine/Source/Programs/AutomationTool/XLocLocalization/XLocLocalization.Automation.csproj
--- a/Engine/Source/Programs/AutomationTool/XLocLocalization/XLocLocalization.Automation.csproj
+++ b/Engine/Source/Programs/AutomationTool/XLocLocalization/XLocLocalization.Automation.csproj
@@ -2,13 +2,9 @@
   <Import Project="..\..\Shared\UnrealEngine.csproj.props" />
   
   <PropertyGroup>
-    <TargetFramework>net8.0</TargetFramework>
-    <Configuration Condition=" '$(Configuration)' == '' ">Development</Configuration>
     <OutputType>Library</OutputType>
-    <AppendTargetFrameworkToOutputPath>false</AppendTargetFrameworkToOutputPath>
     <GenerateAssemblyInfo>false</GenerateAssemblyInfo>
     <GenerateTargetFrameworkAttribute>false</GenerateTargetFrameworkAttribute>
-    <Configurations>Debug;Release;Development</Configurations>
     <RootNamespace>AutomationTool</RootNamespace>
     <AssemblyName>XLocLocalization.Automation</AssemblyName>
     <WarningsNotAsErrors>612,618</WarningsNotAsErrors>
@@ -16,20 +12,6 @@
     <TreatWarningsAsErrors>true</TreatWarningsAsErrors>
     <CopyLocalLockFileAssemblies>true</CopyLocalLockFileAssemblies>
     <SatelliteResourceLanguages>en</SatelliteResourceLanguages> <!-- remove non english resource languages -->
-    <DebugType>pdbonly</DebugType>
-  </PropertyGroup>
-  <PropertyGroup Condition=" '$(Configuration)|$(Platform)' == 'Development|AnyCPU' ">
-    <DefineConstants>$(DefineConstants);TRACE</DefineConstants>
-    <Optimize>true</Optimize>
-  </PropertyGroup>
-  <PropertyGroup Condition=" '$(Configuration)|$(Platform)' == 'Release|AnyCPU' ">
-    <DefineConstants>$(DefineConstants);TRACE</DefineConstants>
-    <Optimize>true</Optimize>
-  </PropertyGroup>
-  <PropertyGroup Condition=" '$(Configuration)|$(Platform)' == 'Debug|AnyCPU' ">
-    <DefineConstants>$(DefineConstants);DEBUG;TRACE</DefineConstants>
-    <DebugSymbols>true</DebugSymbols>
-    <DebugType>full</DebugType>
   </PropertyGroup>
 
   <ItemGroup>
